Fix no logging build issue with SteamSockets

#rnx

diff --git a/Engine/Plugins/Runtime/Steam/SteamSockets/Source/SteamSockets/Private/SteamSocketsSubsystem.cpp b/Engine/Plugins/Runtime/Steam/SteamSockets/Source/SteamSockets/Private/SteamSocketsSubsystem.cpp
--- a/Engine/Plugins/Runtime/Steam/SteamSockets/Source/SteamSockets/Private/SteamSocketsSubsystem.cpp
+++ b/Engine/Plugins/Runtime/Steam/SteamSockets/Source/SteamSockets/Private/SteamSocketsSubsystem.cpp
@@ -21,6 +21,7 @@ DEFINE_LOG_CATEGORY_STATIC(LogSteamSocketsAPI, Log, All);
 /** Message logging hook for debugging Steam Sockets. */
 void SteamSocketsDebugLogger(ESteamNetworkingSocketsDebugOutputType nType, const char *pszMsg)
 {
+#if !NO_LOGGING
 	FString OutputType;
 	ELogVerbosity::Type Verboseness = ELogVerbosity::Verbose;
 	switch (nType)
@@ -57,6 +58,7 @@ void SteamSocketsDebugLogger(ESteamNetworkingSocketsDebugOutputType nType, const
 
 	// Runtime detection of logging level and verbosity.
 	GLog->Log(LogSteamSocketsAPI.GetCategoryName(), Verboseness, FString::Printf(TEXT("SteamSockets API: %s %s"), *OutputType, ANSI_TO_TCHAR(pszMsg)));
+#endif
 }
 
 FSteamSocketsSubsystem* FSteamSocketsSubsystem::SocketSingleton = nullptr;
