Fix rejoining connection issues with SteamSockets due to stagnant handshake data if no new connections are detected

#rnx

diff --git a/Engine/Plugins/Runtime/Steam/SteamSockets/Source/SteamSockets/Private/SteamSocketsNetConnection.cpp b/Engine/Plugins/Runtime/Steam/SteamSockets/Source/SteamSockets/Private/SteamSocketsNetConnection.cpp
--- a/Engine/Plugins/Runtime/Steam/SteamSockets/Source/SteamSockets/Private/SteamSocketsNetConnection.cpp
+++ b/Engine/Plugins/Runtime/Steam/SteamSockets/Source/SteamSockets/Private/SteamSocketsNetConnection.cpp
@@ -204,6 +204,13 @@ void USteamSocketsNetConnection::HandleRecvMessage(void* InData, int32 SizeOfDat
 
 				UE_LOG(LogNet, Log, TEXT("SteamSockets: Connectionless handshake complete"));
 				bInConnectionlessHandshake = false;
+				
+				// Reset the challenge data for the future
+				if (StatelessConnect.IsValid())
+				{
+					StatelessConnect->ResetChallengeData();
+				}
+				
 				SizeOfData = FMath::DivideAndRoundUp(RawPacket.CountBits, 8);
 				if (SizeOfData > 0)
 				{
