Quick NetCore compile fix for UE4 Main.

diff --git a/Engine/Plugins/Runtime/Steam/SteamSockets/Source/SteamSockets/SteamSockets.Build.cs b/Engine/Plugins/Runtime/Steam/SteamSockets/Source/SteamSockets/SteamSockets.Build.cs
--- a/Engine/Plugins/Runtime/Steam/SteamSockets/Source/SteamSockets/SteamSockets.Build.cs
+++ b/Engine/Plugins/Runtime/Steam/SteamSockets/Source/SteamSockets/SteamSockets.Build.cs
@@ -13,6 +13,7 @@ public class SteamSockets : ModuleRules
 			new string[] {
 				"Core", 
 				"CoreUObject",
+				"NetCore",
 				"Engine", 
 				"Sockets",
 				"OnlineSubsystem",
