New support for the SteamSockets networking layer in UE4. This allows a project to use the SteamSocket communication API over the default SteamNetworking API shipped with the Steam OnlineSubsystem.

To Use:

* Enable the SteamSockets plugin in the editor
* Change all NetDriver definitions to use the "/Script/SteamSockets.SteamSocketsNetDriver" class. SteamSockets only works with its own NetDriver. It has support for both IP communications and communications within the Steam platform (via SteamIDs).

diff --git a/Engine/Config/BaseEngine.ini b/Engine/Config/BaseEngine.ini
--- a/Engine/Config/BaseEngine.ini
+++ b/Engine/Config/BaseEngine.ini
@@ -1341,6 +1341,9 @@ P2PConnectionTimeout=90
 [/Script/OnlineSubsystemSteam.SteamNetDriver]
 NetConnectionClassName="/Script/OnlineSubsystemSteam.SteamNetConnection"
 
+[/Script/SteamSockets.SteamSocketsNetDriver]
+NetConnectionClassName="/Script/SteamSockets.SteamSocketsNetConnection"
+
 [OnlineSubsystemMcp]
 bEnabled=false
 bForceEnabledInEditor=false
diff --git a/Engine/Plugins/Online/OnlineSubsystemSteam/Source/Private/OnlineAsyncTaskManagerSteam.cpp b/Engine/Plugins/Online/OnlineSubsystemSteam/Source/Private/OnlineAsyncTaskManagerSteam.cpp
--- a/Engine/Plugins/Online/OnlineSubsystemSteam/Source/Private/OnlineAsyncTaskManagerSteam.cpp
+++ b/Engine/Plugins/Online/OnlineSubsystemSteam/Source/Private/OnlineAsyncTaskManagerSteam.cpp
@@ -667,24 +667,26 @@ public:
 	 */
 	virtual void Finalize() override
 	{
-		FOnlineSessionSteamPtr SessionInt = StaticCastSharedPtr<FOnlineSessionSteam>(Subsystem->GetSessionInterface());
-		if (SessionInt.IsValid())
+		if (Subsystem)
 		{
-			SessionInt->bSteamworksGameServerConnected = true;
-			SessionInt->GameServerSteamId = MakeShareable(new FUniqueNetIdSteam(ServerId));
-
-			FSocketSubsystemSteam* SocketSubsystem = (FSocketSubsystemSteam*)ISocketSubsystem::Get(STEAM_SUBSYSTEM);
-			if (SocketSubsystem)
-			{			
-				SocketSubsystem->FixupSockets(*SessionInt->GameServerSteamId);
-			}
-
-			if (Subsystem)
+			FOnlineSessionSteamPtr SessionInt = StaticCastSharedPtr<FOnlineSessionSteam>(Subsystem->GetSessionInterface());
+			if (SessionInt.IsValid())
 			{
+				SessionInt->bSteamworksGameServerConnected = true;
+				SessionInt->GameServerSteamId = MakeShareable(new FUniqueNetIdSteam(ServerId));
+				if (Subsystem->IsUsingSteamNetworking())
+				{
+					FSocketSubsystemSteam* SocketSubsystem = (FSocketSubsystemSteam*)ISocketSubsystem::Get(STEAM_SUBSYSTEM);
+					if (SocketSubsystem)
+					{
+						SocketSubsystem->FixupSockets(*SessionInt->GameServerSteamId);
+					}
+				}
+
 				Subsystem->TriggerOnSteamServerLoginCompletedDelegates(true);
 			}
-		}
 
+		}
 		// log on is not finished until OnPolicyResponse() is called
 	}
 };
@@ -1012,12 +1014,15 @@ public:
 	 */
 	virtual void Finalize() override
 	{
-		FSocketSubsystemSteam* SocketSubsystem = (FSocketSubsystemSteam*)ISocketSubsystem::Get(STEAM_SUBSYSTEM);
-		if (SocketSubsystem)
+		if (Subsystem && Subsystem->IsUsingSteamNetworking())
 		{
-			if (!SocketSubsystem->AcceptP2PConnection(SteamNetworkingPtr, RemoteId))
-			{	
-				UE_LOG_ONLINE(Log, TEXT("Rejected P2P connection request from %s"), *RemoteId.ToDebugString());
+			FSocketSubsystemSteam* SocketSubsystem = (FSocketSubsystemSteam*)ISocketSubsystem::Get(STEAM_SUBSYSTEM);
+			if (SocketSubsystem)
+			{
+				if (!SocketSubsystem->AcceptP2PConnection(SteamNetworkingPtr, RemoteId))
+				{
+					UE_LOG_ONLINE(Log, TEXT("Rejected P2P connection request from %s"), *RemoteId.ToDebugString());
+				}
 			}
 		}
 	}
@@ -1067,11 +1072,14 @@ public:
 	 */
 	virtual void Finalize() override
 	{
-		// Mark the relevant sockets with this failure so they can properly notify higher level engine code
-		FSocketSubsystemSteam* SocketSubsystem = (FSocketSubsystemSteam*)ISocketSubsystem::Get(STEAM_SUBSYSTEM);
-		if (SocketSubsystem)
+		if (Subsystem && Subsystem->IsUsingSteamNetworking())
 		{
-			SocketSubsystem->ConnectFailure(RemoteId);
+			// Mark the relevant sockets with this failure so they can properly notify higher level engine code
+			FSocketSubsystemSteam* SocketSubsystem = (FSocketSubsystemSteam*)ISocketSubsystem::Get(STEAM_SUBSYSTEM);
+			if (SocketSubsystem)
+			{
+				SocketSubsystem->ConnectFailure(RemoteId);
+			}
 		}
 	}
 };
diff --git a/Engine/Plugins/Online/OnlineSubsystemSteam/Source/Public/OnlineSubsystemSteam.h b/Engine/Plugins/Online/OnlineSubsystemSteam/Source/Public/OnlineSubsystemSteam.h
--- a/Engine/Plugins/Online/OnlineSubsystemSteam/Source/Public/OnlineSubsystemSteam.h
+++ b/Engine/Plugins/Online/OnlineSubsystemSteam/Source/Public/OnlineSubsystemSteam.h
@@ -341,6 +341,15 @@ public:
 		return GameServerQueryPort;
 	}
 
+	/**
+	 *	@return if this subsystem is using SteamNetworking functionality 
+	 *			or another network layer like SteamSockets
+	 */
+	inline bool IsUsingSteamNetworking() const
+	{
+		return bUsingSteamNetworking;
+	}
+
 	/**
 	 * This delegate fires whenever a steam login has succeeded or failed its async task.
 	 * Useful for modules that need to check to see if a user is logged in before running other behavior
diff --git a/Engine/Plugins/Runtime/Steam/SteamSockets/Source/SteamSockets/Private/IPAddressSteamSockets.cpp b/Engine/Plugins/Runtime/Steam/SteamSockets/Source/SteamSockets/Private/IPAddressSteamSockets.cpp
new file mode 100644
--- /dev/null
+++ b/Engine/Plugins/Runtime/Steam/SteamSockets/Source/SteamSockets/Private/IPAddressSteamSockets.cpp
@@ -0,0 +1,239 @@
+// Copyright 1998-2019 Epic Games, Inc. All Rights Reserved.
+
+#include "IPAddressSteamSockets.h"
+#include "SteamSocketsPrivate.h"
+#include "SteamSocketsTypes.h"
+#include "Algo/Reverse.h"
+
+TArray<uint8> FInternetAddrSteamSockets::GetRawIp() const
+{
+	// Due to some API design decisions, we cannot rely on something like the generic bytes member
+	// Instead we have to create the raw data ourselves.
+	TArray<uint8> RawAddressArray;
+
+	// Same code from IPAddressSteam.
+	if (Addr.m_eType == k_ESteamNetworkingIdentityType_SteamID)
+	{
+		uint64 SteamIDNum = Addr.GetSteamID64();
+		const uint8* SteamIdWalk = (uint8*)&SteamIDNum;
+		while (RawAddressArray.Num() < sizeof(uint64))
+		{
+			RawAddressArray.Add(*SteamIdWalk);
+			++SteamIdWalk;
+		}
+
+		// We want to make sure that these arrays are in big endian format.
+#if PLATFORM_LITTLE_ENDIAN
+		Algo::Reverse(RawAddressArray);
+#endif
+
+		RawAddressArray.EmplaceAt(0, k_ESteamNetworkingIdentityType_SteamID); // Push in the type
+	}
+	else if (Addr.m_eType == k_ESteamNetworkingIdentityType_IPAddress && Addr.GetIPAddr())
+	{
+		const SteamNetworkingIPAddr* RawSteamIP = Addr.GetIPAddr();
+		RawAddressArray.Add(k_ESteamNetworkingIdentityType_IPAddress);
+		for (int32 i = 0; i < ARRAY_COUNT(RawSteamIP->m_ipv6); ++i)
+		{
+			RawAddressArray.Add(RawSteamIP->m_ipv6[i]);
+		}
+	}
+
+	return RawAddressArray;
+}
+
+void FInternetAddrSteamSockets::SetRawIp(const TArray<uint8>& RawAddr)
+{
+	// Array is invalid, the type should be the first thing in the array.
+	if (RawAddr.Num() <= 1)
+	{
+		return;
+	}
+
+	Addr.Clear();
+	uint8 ArrayType = RawAddr[0];
+
+	if (ArrayType == k_ESteamNetworkingIdentityType_SteamID)
+	{
+		TArray<uint8> WorkingArray = RawAddr;
+		WorkingArray.RemoveAt(0); // Remove the type
+
+#if PLATFORM_LITTLE_ENDIAN
+		Algo::Reverse(WorkingArray);
+#endif
+		uint64 NewSteamId = 0;
+		for (int32 i = 0; i < WorkingArray.Num(); ++i)
+		{
+			NewSteamId |= (uint64)WorkingArray[i] << (i * 8);
+		}
+
+		Addr.SetSteamID64(NewSteamId);
+	}
+	else if(ArrayType == k_ESteamNetworkingIdentityType_IPAddress)
+	{
+		SteamNetworkingIPAddr NewAddr;
+		NewAddr.Clear();
+
+		// Skip over the flag bit and copy the array into the structure entirely.
+		for (int32 i = 1; i < RawAddr.Num(); ++i)
+		{
+			NewAddr.m_ipv6[i - 1] = RawAddr[i];
+		}
+
+		Addr.SetIPAddr(NewAddr);
+	}
+}
+
+void FInternetAddrSteamSockets::SetIp(const TCHAR* InAddr, bool& bIsValid)
+{
+	FString InAddrStr(InAddr);
+	// See if we're a SteamID.
+	if (InAddrStr.StartsWith(STEAM_URL_PREFIX))
+	{
+		InAddrStr.RemoveFromStart(STEAM_URL_PREFIX);
+		FString SteamIPStr, SteamChannelStr;
+		// Look for a virtual port
+		if (InAddrStr.Split(":", &SteamIPStr, &SteamChannelStr, ESearchCase::CaseSensitive, ESearchDir::FromEnd))
+		{
+			const uint64 Id = FCString::Atoi64(*SteamIPStr);
+			if (Id != 0)
+			{
+				Addr.SetSteamID64(Id);
+				const int32 Channel = FCString::Atoi(*SteamChannelStr);
+				if (Channel != 0 || SteamChannelStr == "0")
+				{
+					P2PVirtualPort = Channel;
+					bIsValid = true;
+				}
+			}
+			else
+			{
+				bIsValid = false;
+				return;
+			}
+		}
+		else
+		{
+			Addr.SetSteamID64(FCString::Atoi64(*InAddrStr));
+			bIsValid = true;
+			return;
+		}
+	}
+	else if(SteamNetworkingUtils()) // check this because the SteamAPI does not.
+	{
+		// This is an IP address. Tell Steam to parse it.
+		SteamNetworkingIPAddr NewAddress;
+		bIsValid = NewAddress.ParseString(TCHAR_TO_ANSI(InAddr));
+		Addr.SetIPAddr(NewAddress);
+	}
+}
+
+void FInternetAddrSteamSockets::SetPort(int32 InPort)
+{
+	if (GetProtocolType() == FNetworkProtocolTypes::SteamSocketsP2P)
+	{
+		P2PVirtualPort = InPort;
+	}
+	else
+	{
+		// There's no way to set a port unless you recreate the entire object.
+		// It's a restricted property, with no setters, so we have to jump through some hoops.
+		SteamNetworkingIPAddr NewAddressInfo;
+		// We don't have to check if the data returned here is valid 
+		// as GetProtocolType will have already done it for us.
+		const SteamNetworkingIPAddr* InternalAddress = Addr.GetIPAddr();
+		FMemory::Memcpy(NewAddressInfo, *InternalAddress);
+		NewAddressInfo.m_port = InPort;
+		Addr.SetIPAddr(NewAddressInfo);
+	}
+}
+
+int32 FInternetAddrSteamSockets::GetPort() const
+{
+	if (GetProtocolType() == FNetworkProtocolTypes::SteamSocketsP2P)
+	{
+		return P2PVirtualPort;
+	}
+	else if (const SteamNetworkingIPAddr* InternalAddress = Addr.GetIPAddr())
+	{
+		// Valve stores this in host byte order as of writing.
+		return InternalAddress->m_port;
+	}
+
+	return 0;
+}
+
+void FInternetAddrSteamSockets::SetAnyAddress()
+{
+	// This will remove all data such that we can set an any address
+	Addr.Clear();
+	SteamNetworkingIPAddr NewAddress;
+	NewAddress.Clear();
+	Addr.SetIPAddr(NewAddress);
+}
+
+FString FInternetAddrSteamSockets::ToString(bool bAppendPort) const
+{
+	// Attempt to print out only if we have valid data.
+	if (!Addr.IsInvalid())
+	{
+		if (GetProtocolType() == FNetworkProtocolTypes::SteamSocketsP2P)
+		{
+			FString BaseResult = FString::Printf(TEXT("%s%llu"), STEAM_URL_PREFIX, Addr.GetSteamID64());
+			if (bAppendPort)
+			{
+				BaseResult += FString::Printf(TEXT(":%d"), P2PVirtualPort);
+			}
+			return BaseResult;
+		}
+
+		const SteamNetworkingIPAddr* InternalAddress = Addr.GetIPAddr();
+		if (InternalAddress && SteamNetworkingUtils())
+		{
+			ANSICHAR StrBuffer[SteamNetworkingIPAddr::k_cchMaxString];
+			FMemory::Memzero(StrBuffer);
+			InternalAddress->ToString(StrBuffer, SteamNetworkingIPAddr::k_cchMaxString, bAppendPort);
+
+			return ANSI_TO_TCHAR(StrBuffer);
+		}
+	}
+
+	return TEXT("Invalid");
+}
+
+uint32 FInternetAddrSteamSockets::GetTypeHash() const
+{
+	// This object doesn't have any address data at all
+	// we should attempt to hash something in it anyways
+	if (Addr.IsInvalid())
+	{
+		// This object is completely blank, we can't do anything at all.
+		if (P2PVirtualPort == 0)
+		{
+			return 0;
+		}
+		else
+		{
+			// Otherwise, hash the virtual port.
+			return ::GetTypeHash(P2PVirtualPort);
+		}
+	}
+
+	if (GetProtocolType() == FNetworkProtocolTypes::SteamSocketsP2P)
+	{
+		return ::GetTypeHash(Addr.GetSteamID64());
+	}
+	return ::GetTypeHash(ToString(true));
+}
+
+FName FInternetAddrSteamSockets::GetProtocolType() const
+{
+	// By default, an empty Addr structure will have a type of unknown, 
+	// which will cause this function to return nullptr.
+	if (Addr.GetIPAddr() == nullptr)
+	{
+		return FNetworkProtocolTypes::SteamSocketsP2P;
+	}
+
+	return FNetworkProtocolTypes::SteamSocketsIP;
+}
diff --git a/Engine/Plugins/Runtime/Steam/SteamSockets/Source/SteamSockets/Private/IPAddressSteamSockets.h b/Engine/Plugins/Runtime/Steam/SteamSockets/Source/SteamSockets/Private/IPAddressSteamSockets.h
new file mode 100644
--- /dev/null
+++ b/Engine/Plugins/Runtime/Steam/SteamSockets/Source/SteamSockets/Private/IPAddressSteamSockets.h
@@ -0,0 +1,169 @@
+// Copyright 1998-2019 Epic Games, Inc. All Rights Reserved.
+
+#pragma once
+
+#include "CoreMinimal.h"
+#include "IPAddress.h"
+#include "SteamSocketsPrivate.h"
+#include "SteamSocketsPackage.h"
+#include "SteamSocketsTypes.h"
+#include "Engine/EngineBaseTypes.h"
+
+class FInternetAddrSteamSockets : public FInternetAddr
+{
+PACKAGE_SCOPE:
+	SteamNetworkingIdentity Addr;
+	int32 P2PVirtualPort;
+
+public:
+	FInternetAddrSteamSockets() :
+		P2PVirtualPort(0)
+	{
+		Addr.Clear();
+	}
+
+	FInternetAddrSteamSockets(const FInternetAddrSteamSockets& In) :
+		Addr(In.Addr),
+		P2PVirtualPort(In.P2PVirtualPort)
+	{
+	}
+
+	FInternetAddrSteamSockets(const SteamNetworkingIdentity& NewAddress) :
+		Addr(NewAddress),
+		P2PVirtualPort(0)
+	{
+	}
+
+	FInternetAddrSteamSockets(const SteamNetworkingIPAddr& IPAddr) :
+		P2PVirtualPort(0)
+	{
+		Addr.SetIPAddr(IPAddr);
+	}
+
+	explicit FInternetAddrSteamSockets(uint64& SteamID) :
+		P2PVirtualPort(0)
+	{
+		Addr.SetSteamID64(SteamID);
+	}
+	
+	virtual TArray<uint8> GetRawIp() const override;
+
+	virtual void SetRawIp(const TArray<uint8>& RawAddr) override;
+
+	virtual void SetIp(uint32 InAddr) override
+	{
+		/** Not used */
+	}
+
+	/**
+	* Sets the ip address from a string ("A.B.C.D") or a steam id "steam.STEAMID" or "STEAMID:PORT"
+	*
+	* @param InAddr the string containing the new ip address to use
+	*/
+	virtual void SetIp(const TCHAR* InAddr, bool& bIsValid) override;
+
+	/**
+	* Copies the network byte order ip address to a host byte order dword
+	*
+	* @param OutAddr the out param receiving the ip address
+	*/
+	virtual void GetIp(uint32& OutAddr) const override
+	{
+		/** Not used */
+	}
+
+	/**
+	* Sets the port number from a host byte order int
+	*
+	* @param InPort the new port to use (must convert to network byte order)
+	*/
+	virtual void SetPort(int32 InPort) override;
+
+	/**
+	* Returns the port number from this address in host byte order
+	*/
+	virtual int32 GetPort() const override;
+
+	/**
+	 * Set Platform specific port data
+	 */
+	virtual void SetPlatformPort(int32 InPort) override
+	{
+		P2PVirtualPort = (int16)InPort;
+	}
+
+	/**
+	 * Get platform specific port data.
+	 */
+	virtual int32 GetPlatformPort() const override
+	{
+		return (int32)P2PVirtualPort;
+	}
+
+	/** Sets the address to be any address */
+	virtual void SetAnyAddress() override;
+
+	/** Sets the address to broadcast */
+	virtual void SetBroadcastAddress() override
+	{
+		/** Not used */
+	}
+
+	/** Sets the address to loopback */
+	virtual void SetLoopbackAddress() override
+	{
+		Addr.SetLocalHost();
+	}
+
+	/**
+	* Converts this internet ip address to string form
+	*
+	* @param bAppendPort whether to append the port information or not
+	*/
+	virtual FString ToString(bool bAppendPort) const override;
+
+	/**
+	* Compares two internet ip addresses for equality
+	*
+	* @param Other the address to compare against
+	*/
+	virtual bool operator==(const FInternetAddr& Other) const override
+	{
+		FInternetAddrSteamSockets& SteamOther = (FInternetAddrSteamSockets&)Other;
+		return Addr == SteamOther.Addr;
+	}
+
+	virtual uint32 GetTypeHash() const override;
+
+	virtual FName GetProtocolType() const override;
+
+	virtual bool IsValid() const override
+	{
+		return !(Addr.IsInvalid());
+	}
+
+	operator const SteamNetworkingIPAddr() const
+	{
+		const SteamNetworkingIPAddr* IPAddr = Addr.GetIPAddr();
+		if (IPAddr == nullptr)
+		{
+			SteamNetworkingIPAddr EmptyAddr;
+			EmptyAddr.Clear();
+			return EmptyAddr;
+		}
+		return *IPAddr;
+	}
+
+	operator const SteamNetworkingIdentity() const
+	{
+		return Addr;
+	}
+
+	virtual TSharedRef<FInternetAddr> Clone() const override
+	{
+		TSharedRef<FInternetAddrSteamSockets> NewAddress = MakeShareable(new FInternetAddrSteamSockets);
+		NewAddress->Addr = Addr;
+		NewAddress->P2PVirtualPort = P2PVirtualPort;
+		return NewAddress;
+	}
+};
diff --git a/Engine/Plugins/Runtime/Steam/SteamSockets/Source/SteamSockets/Private/SteamSocket.cpp b/Engine/Plugins/Runtime/Steam/SteamSockets/Source/SteamSockets/Private/SteamSocket.cpp
new file mode 100644
--- /dev/null
+++ b/Engine/Plugins/Runtime/Steam/SteamSockets/Source/SteamSockets/Private/SteamSocket.cpp
@@ -0,0 +1,594 @@
+// Copyright 1998-2019 Epic Games, Inc. All Rights Reserved.
+
+#pragma once
+
+#include "SteamSocket.h"
+#include "CoreMinimal.h"
+#include "SteamSocketsTypes.h"
+#include "SteamSocketsSubsystem.h"
+#include "SteamSocketsPrivate.h"
+#include "SocketTypes.h"
+
+FSteamSocket::FSteamSocket(ESocketType InSocketType, const FString& InSocketDescription, const FName& InSocketProtocol) :
+	FSocket(InSocketType, InSocketDescription, InSocketProtocol),
+	InternalHandle(k_HSteamNetConnection_Invalid),
+	SendMode(k_nSteamNetworkingSend_UnreliableNoNagle),
+	bShouldLingerOnClose(false),
+	bIsListenSocket(false),
+	bIsLANSocket(false),
+	bHasPendingData(false),
+	ClosureReason(k_ESteamNetConnectionEnd_App_Generic)
+{
+	SocketSubsystem = static_cast<FSteamSocketsSubsystem*>(ISocketSubsystem::Get(STEAM_SOCKETS_SUBSYSTEM));
+}
+
+FSteamSocket::~FSteamSocket()
+{
+	// Release any data that is still around.
+	if (bHasPendingData && PendingData != nullptr)
+	{
+		PendingData->Release();
+	}
+	
+	Close();
+}
+
+bool FSteamSocket::Close()
+{
+	bool bWasSuccessful = false;
+
+	// If we're already closed, don't bother doing anything else.
+	if (InternalHandle == k_HSteamListenSocket_Invalid || InternalHandle == k_HSteamNetConnection_Invalid)
+	{
+		UE_LOG(LogSockets, VeryVerbose, TEXT("SteamSockets: Socket is already cleaned up, ready for destruction."));
+		return true;
+	}
+
+	ISteamNetworkingSockets* SocketInterface = FSteamSocketsSubsystem::GetSteamSocketsInterface();
+
+	// Safety check for the interface.
+	if (SocketInterface == nullptr)
+	{
+		UE_LOG(LogSockets, VeryVerbose, TEXT("SteamSockets: Socket interface not found, cannot close."));
+		return false;
+	}
+
+	if (bIsListenSocket)
+	{
+		// Closing a listen socket destroys any sockets that were accepted from this socket
+		bWasSuccessful = SocketInterface->CloseListenSocket(InternalHandle);
+		if (bWasSuccessful)
+		{
+			UE_LOG(LogSockets, Verbose, TEXT("SteamSockets: Closed listen socket %u via the API"), InternalHandle);
+			SocketSubsystem->RemoveSocketsForListener(this);
+			InternalHandle = k_HSteamListenSocket_Invalid;
+		}
+	}
+	else
+	{
+		bWasSuccessful = SocketInterface->CloseConnection(InternalHandle, (int32)ClosureReason, "Connection Ended.", bShouldLingerOnClose);
+		if (bWasSuccessful)
+		{
+			UE_LOG(LogSockets, Verbose, TEXT("SteamSockets: Closed socket %u via the API, reason %d"), InternalHandle, ClosureReason);
+			SocketSubsystem->QueueRemoval(InternalHandle);			
+			InternalHandle = k_HSteamNetConnection_Invalid;
+		}
+	}
+
+	return bWasSuccessful;
+}
+
+bool FSteamSocket::Bind(const FInternetAddr& Addr)
+{
+	if (Addr.GetProtocolType() == FNetworkProtocolTypes::SteamSocketsIP || 
+		Addr.GetProtocolType() == FNetworkProtocolTypes::SteamSocketsP2P)
+	{
+		// Because of the SteamAPI, we have to just hold onto the bind address
+		// We can't actually do anything with it until later.
+		BindAddress = *((FInternetAddrSteamSockets*)&Addr);
+		return true;
+	}
+
+	return false;
+}
+
+bool FSteamSocket::Connect(const FInternetAddr& Addr)
+{
+	ISteamNetworkingSockets* SocketInterface = FSteamSocketsSubsystem::GetSteamSocketsInterface();
+
+	// Check if the respective interface is valid
+	if (SocketInterface == nullptr)
+	{
+		UE_LOG(LogSockets, Warning, TEXT("SteamSockets: Socket Interface is null, cannot connect"));
+		return false;
+	}
+
+	if (!Addr.IsValid())
+	{
+		UE_LOG(LogSockets, Warning, TEXT("SteamSockets: Connection address is not valid. Cannot connect!"));
+		SocketSubsystem->LastSocketError = SE_EINVAL;
+		return false;
+	}
+
+	FInternetAddrSteamSockets SteamAddr = *((FInternetAddrSteamSockets*)&Addr);
+	if (GetProtocol() == FNetworkProtocolTypes::SteamSocketsIP)
+	{
+		InternalHandle = SocketInterface->ConnectByIPAddress(SteamAddr);
+		SetLanOptions();
+
+	}
+	else if (GetProtocol() == FNetworkProtocolTypes::SteamSocketsP2P)
+	{
+		InternalHandle = SocketInterface->ConnectP2P(SteamAddr, SteamAddr.GetPort());
+	}
+
+	if (InternalHandle != k_HSteamNetConnection_Invalid)
+	{
+		UE_LOG(LogSockets, Verbose, TEXT("SteamSockets: Connection to %s initiated"), *Addr.ToString(false));
+		SocketSubsystem->AddSocket(Addr, this);
+		return true;
+	}
+
+	return false;
+}
+
+bool FSteamSocket::Listen(int32 MaxBacklog)
+{
+	ISteamNetworkingSockets* SocketInterface = FSteamSocketsSubsystem::GetSteamSocketsInterface();
+
+	// Require the interface to be non-null
+	if (SocketInterface  == nullptr)
+	{
+		UE_LOG(LogSockets, Warning, TEXT("SteamSockets: Socket interface is null, cannot establish interface for listening!"));
+		return false;
+	}
+
+	bool bWasSuccessful = false;
+	bIsListenSocket = true;
+	if (GetProtocol() == FNetworkProtocolTypes::SteamSocketsIP)
+	{
+		InternalHandle = SocketInterface->CreateListenSocketIP(BindAddress);
+		SetLanOptions();
+	}
+	else
+	{
+		InternalHandle = SocketInterface->CreateListenSocketP2P(BindAddress.GetPlatformPort());
+	}
+
+	if (InternalHandle != k_HSteamListenSocket_Invalid)
+	{
+		bWasSuccessful = true;
+		SocketSubsystem->AddSocket(BindAddress, this);
+	}
+
+	return bWasSuccessful;
+}
+
+class FSocket* FSteamSocket::Accept(const FString& InSocketDescription)
+{
+	FSteamSocket* NewSocket = static_cast<FSteamSocket*>(SocketSubsystem->CreateSocket(FName("SteamClientSocket"), InSocketDescription, GetProtocol()));
+	NewSocket->SendMode = SendMode;
+	NewSocket->bShouldLingerOnClose = bShouldLingerOnClose;
+	NewSocket->bIsLANSocket = bIsLANSocket;
+
+	// Set virtual ports here to preserve channel information over the P2P network.
+	if (BindAddress.GetProtocolType() == FNetworkProtocolTypes::SteamSocketsP2P)
+	{
+		NewSocket->BindAddress.SetPlatformPort(BindAddress.GetPlatformPort());
+	}
+
+	return NewSocket;
+}
+
+bool FSteamSocket::SendTo(const uint8* Data, int32 Count, int32& BytesSent, const FInternetAddr& Destination)
+{
+	FInternetAddrSteamSockets CurrentAddress;
+	if (GetPeerAddress(CurrentAddress) && Destination == CurrentAddress)
+	{
+		return Send(Data, Count, BytesSent);
+	}
+	else if (SocketSubsystem)
+	{
+		FSteamSocketsSubsystem::FSteamSocketInformation* SocketInfo = SocketSubsystem->GetSocketInfo(Destination);
+		// IsValid will also check the validity of the socket pointer as well
+		if (SocketInfo != nullptr && SocketInfo->IsValid())
+		{
+			return SocketInfo->Socket->Send(Data, Count, BytesSent);
+		}
+	}
+
+	return false;
+}
+
+bool FSteamSocket::Send(const uint8* Data, int32 Count, int32& BytesSent)
+{
+	BytesSent = 0;
+	// Steam won't send packets if we are not marked as connected.
+	if (InternalHandle != k_HSteamNetConnection_Invalid && GetConnectionState() == SCS_Connected)
+	{
+		// GetConnectionState will check the validity of the sockets interface for us.
+		switch (FSteamSocketsSubsystem::GetSteamSocketsInterface()->SendMessageToConnection(InternalHandle, (void*)Data, Count, SendMode))
+		{
+			case k_EResultOK:
+				SocketSubsystem->LastSocketError = SE_NO_ERROR;
+				BytesSent = Count;
+				return true;
+			break;
+			case k_EResultInvalidParam:
+				SocketSubsystem->LastSocketError = SE_EINVAL;
+			break;
+			case k_EResultInvalidState:
+				SocketSubsystem->LastSocketError = SE_EBADF;
+			break;
+			case k_EResultNoConnection:
+				SocketSubsystem->LastSocketError = SE_ENOTCONN;
+			break;
+			case k_EResultIgnored:
+				SocketSubsystem->LastSocketError = SE_SYSNOTREADY;
+			break;
+			case k_EResultLimitExceeded:
+				SocketSubsystem->LastSocketError = SE_EPROCLIM;
+			break;
+			default:
+				SocketSubsystem->LastSocketError = SE_EFAULT;
+			break;
+		}
+		BytesSent = -1;
+	}
+
+	return false;
+}
+
+bool FSteamSocket::Recv(uint8* Data, int32 BufferSize, int32& BytesRead, ESocketReceiveFlags::Type Flags)
+{
+	BytesRead = -1;
+	SteamNetworkingMessage_t* Message;
+	int32 MessagesRead;
+
+	if (RecvRaw(Message, 1, MessagesRead, Flags))
+	{
+		if (MessagesRead >= 1)
+		{
+			if (BufferSize >= 0 && Message->GetSize() <= (uint32)BufferSize)
+			{
+				FMemory::Memcpy(Data, Message->GetData(), BytesRead);
+				Message->Release();
+			}
+			else
+			{
+				BytesRead = -1;
+				SocketSubsystem->LastSocketError = SE_EMSGSIZE;
+				return false;
+			}		
+		}
+		else
+		{
+			BytesRead = 0;
+		}
+		return true;		
+	}
+
+	return false;
+}
+
+bool FSteamSocket::RecvRaw(SteamNetworkingMessage_t*& Data, int32 MaxMessages, int32& MessagesRead, ESocketReceiveFlags::Type Flags)
+{
+	ISteamNetworkingSockets* SocketInterface = FSteamSocketsSubsystem::GetSteamSocketsInterface();
+	if (SocketInterface == nullptr)
+	{
+		SocketSubsystem->LastSocketError = SE_SYSNOTREADY;
+		return false;
+	}
+
+	// Connection doesn't exist, return false
+	if (InternalHandle == k_HSteamNetConnection_Invalid || (bIsListenSocket && InternalHandle == k_HSteamListenSocket_Invalid))
+	{
+		SocketSubsystem->LastSocketError = SE_EINVAL;
+		return false;
+	}
+
+	bool bIsPeeking = (Flags == ESocketReceiveFlags::Peek);
+	if (bIsPeeking)
+	{
+		// If someone asks for more than 1 message, this will break the message queue.
+		if (MaxMessages > 1)
+		{
+			UE_LOG(LogSockets, Warning, TEXT("SteamSockets: Recv cannot peek more than 1 message on the queue. Was asked for %d"), MaxMessages);
+			MaxMessages = 1;
+		}
+
+		// Don't attempt to peek again if we already have pending data as that will mess up the message queue.
+		if (bHasPendingData)
+		{
+			MessagesRead = 1;
+			return true;
+		}
+	}
+	else if (bHasPendingData)
+	{
+		// Set the pointer value
+		Data = PendingData;
+		// null whatever we are currently holding
+		PendingData = nullptr;
+		MessagesRead = 1;
+		bHasPendingData = false;
+		return true;
+	}
+
+	// At this point, we will have already written our pending data or we're getting a new one.
+	MessagesRead = (bIsListenSocket) ? SocketInterface->ReceiveMessagesOnListenSocket(InternalHandle, ((bIsPeeking) ? &PendingData : &Data), MaxMessages) :
+		SocketInterface->ReceiveMessagesOnConnection(InternalHandle, ((bIsPeeking) ? &PendingData : &Data), MaxMessages);
+
+	if (MessagesRead >= 1)
+	{
+		if (bIsPeeking)
+		{
+			bHasPendingData = true;
+		}
+		SocketSubsystem->LastSocketError = SE_NO_ERROR;
+		return true;
+	}
+	else if (MessagesRead == 0)
+	{
+		bHasPendingData = false;
+		PendingData = nullptr;
+		return true;
+	}
+	else
+	{
+		UE_LOG(LogSockets, Error, TEXT("SteamSockets: Recv Connection handle is marked as invalid! Is Listen Socket? %d"), bIsListenSocket);
+		MessagesRead = -1;
+	}
+
+	SocketSubsystem->LastSocketError = SE_EFAULT;
+	return false;
+}
+
+bool FSteamSocket::HasPendingData(uint32& PendingDataSize)
+{
+	if (bHasPendingData)
+	{
+		PendingDataSize = PendingData->GetSize();
+		return true;
+	}
+
+	int32 MessagesRead;
+	PendingDataSize = 0;
+	// We still have to pass a valid pointer reference, but this will be unwritten to and unused.
+	SteamNetworkingMessage_t* FakeMessage = nullptr;
+
+	if (RecvRaw(FakeMessage, 1, MessagesRead, ESocketReceiveFlags::Peek))
+	{
+		if (MessagesRead >= 1 && PendingData != nullptr)
+		{
+			PendingDataSize = PendingData->GetSize();
+			return true;
+		}
+	}
+	
+	return false;
+}
+
+ESocketConnectionState FSteamSocket::GetConnectionState()
+{
+	// Early out if we're not connected.
+	ISteamNetworkingSockets* SocketInterface = FSteamSocketsSubsystem::GetSteamSocketsInterface();
+	if (InternalHandle == k_HSteamNetConnection_Invalid || SocketInterface == nullptr)
+	{
+		return SCS_NotConnected;
+	}
+
+	SteamNetworkingQuickConnectionStatus QuickSocketData;
+	if (SocketInterface->GetQuickConnectionStatus(InternalHandle, &QuickSocketData))
+	{
+		switch (QuickSocketData.m_eState)
+		{
+		case k_ESteamNetworkingConnectionState_Connected:
+			return SCS_Connected;
+		case k_ESteamNetworkingConnectionState_None:
+		case k_ESteamNetworkingConnectionState_ProblemDetectedLocally:
+			return SCS_ConnectionError;
+		default:
+			return SCS_NotConnected;
+		}
+	}
+
+	return SCS_NotConnected;
+}
+
+void FSteamSocket::SetSendMode(int32 NewSendMode)
+{
+	SendMode = NewSendMode;
+}
+
+bool FSteamSocket::GetPeerAddress(FInternetAddr& OutAddr)
+{
+	ISteamNetworkingSockets* SocketInterface = FSteamSocketsSubsystem::GetSteamSocketsInterface();
+	FInternetAddrSteamSockets& SteamAddr = (FInternetAddrSteamSockets&)OutAddr;
+	SteamNetConnectionInfo_t CurrentConnectionInfo;
+	if (SocketInterface && SocketInterface->GetConnectionInfo(InternalHandle, &CurrentConnectionInfo))
+	{
+		if (CurrentConnectionInfo.m_identityRemote.IsInvalid())
+		{
+			// There is no way to check validity on the remote addr structure
+			// So the next best thing is to check if it's all zeros.
+			if (CurrentConnectionInfo.m_addrRemote.IsIPv6AllZeros())
+			{
+				return false;
+			}
+
+			SteamAddr = CurrentConnectionInfo.m_addrRemote;
+		}
+		else
+		{
+			SteamAddr = CurrentConnectionInfo.m_identityRemote;
+			// Remember to set the platform port here.
+			SteamAddr.SetPlatformPort(BindAddress.GetPlatformPort());
+		}
+		return true;
+	}
+
+	return false;
+}
+
+bool FSteamSocket::SetNoDelay(bool bIsNoDelay)
+{
+	// See https://partner.steamgames.com/doc/api/steamnetworkingtypes#message_sending_flags for what all these flags mean
+	if (bIsNoDelay)
+	{
+		if (SendMode == k_nSteamNetworkingSend_Unreliable || SendMode == k_nSteamNetworkingSend_NoNagle || SendMode == k_nSteamNetworkingSend_UnreliableNoNagle)
+		{
+			SendMode = k_nSteamNetworkingSend_UnreliableNoDelay;
+		}
+		else if (SendMode == k_nSteamNetworkingSend_Reliable)
+		{
+			SendMode = k_nSteamNetworkingSend_ReliableNoNagle;
+		}
+	}
+	else
+	{
+		if (SendMode == k_nSteamNetworkingSend_NoDelay)
+		{
+			SendMode = k_nSteamNetworkingSend_Unreliable;
+		}
+		else if (SendMode == k_nSteamNetworkingSend_UnreliableNoDelay)
+		{
+			SendMode = k_nSteamNetworkingSend_UnreliableNoNagle;
+		}
+		else if (SendMode == k_nSteamNetworkingSend_ReliableNoNagle) // Not necessarily a NoDelay, but this is the closest thing.
+		{
+			SendMode = k_nSteamNetworkingSend_Reliable;
+		}
+	}
+
+	return true;
+}
+
+bool FSteamSocket::SetLinger(bool bShouldLinger, int32 Timeout)
+{
+	bShouldLingerOnClose = bShouldLinger;
+	return true;
+}
+
+void FSteamSocket::SetLanOptions()
+{
+	// This only should get set if you're in a LAN environment. 
+	// The Netdriver allows us to set a lan socket flag which we'll use later for delayed listens.
+	if (SteamNetworkingUtils() && bIsLANSocket && InternalHandle != k_HSteamNetConnection_Invalid)
+	{
+		SteamNetworkingUtils()->SetConnectionConfigValueInt32(InternalHandle, k_ESteamNetworkingConfig_IP_AllowWithoutAuth, 1);
+	}
+}
+
+bool FSteamSocket::SetSendBufferSize(int32 Size, int32& NewSize)
+{
+	bool bSuccess = false;
+	if (SteamNetworkingUtils())
+	{
+		if (SteamNetworkingUtils()->SetConnectionConfigValueInt32(InternalHandle, k_ESteamNetworkingConfig_SendBufferSize, Size))
+		{
+			bSuccess = true;
+		}
+
+		// Get the buffer size if we succeed or fail anyways.
+		size_t ConfigValueSize = sizeof(int32);
+		SteamNetworkingUtils()->GetConfigValue(k_ESteamNetworkingConfig_SendBufferSize, k_ESteamNetworkingConfig_Connection, InternalHandle, nullptr, &NewSize, &ConfigValueSize);
+	}
+	else
+	{
+		UE_LOG(LogSockets, Warning, TEXT("SteamSockets: Setting buffer size requires access to network utilities!"));
+		NewSize = -1;
+	}
+
+	return bSuccess;
+}
+
+bool FSteamSocket::SetReceiveBufferSize(int32 Size, int32& NewSize)
+{
+	// This is 1:1 with SendBufferSize on SteamSockets.
+	return SetSendBufferSize(Size, NewSize);
+}
+
+bool FSteamSocket::Shutdown(ESocketShutdownMode Mode)
+{
+	// Unsupported.
+	return false;
+}
+
+bool FSteamSocket::Wait(ESocketWaitConditions::Type Condition, FTimespan WaitTime)
+{
+	// Unsupported. Can't tell if ready to write or not.
+	return false;
+}
+
+bool FSteamSocket::WaitForPendingConnection(bool& bHasPendingConnection, const FTimespan& WaitTime)
+{
+	// Unfortunately, this API has no ability to check if there's a pending connection ahead of time.
+	return false;
+}
+
+bool FSteamSocket::SetReuseAddr(bool bAllowReuse)
+{
+	// Unsupported.
+	return true;
+}
+
+bool FSteamSocket::SetRecvErr(bool bUseErrorQueue)
+{
+	// Unsupported.
+	return true;
+}
+
+bool FSteamSocket::SetNonBlocking(bool bIsNonBlocking)
+{
+	// Unsupported. API does not block.
+	return true;
+}
+
+bool FSteamSocket::SetBroadcast(bool bAllowBroadcast)
+{
+	// Unsupported.
+	return false;
+}
+
+bool FSteamSocket::JoinMulticastGroup(const FInternetAddr& GroupAddress, const FInternetAddr& InterfaceAddress)
+{
+	// Unsupported.
+	return false;
+}
+
+bool FSteamSocket::JoinMulticastGroup(const FInternetAddr& GroupAddress)
+{
+	// Unsupported.
+	return false;
+}
+
+bool FSteamSocket::LeaveMulticastGroup(const FInternetAddr& GroupAddress)
+{
+	// Unsupported.
+	return false;
+}
+
+bool FSteamSocket::LeaveMulticastGroup(const FInternetAddr& GroupAddress, const FInternetAddr& InterfaceAddress)
+{
+	// Unsupported.
+	return false;
+}
+
+bool FSteamSocket::SetMulticastLoopback(bool bLoopback)
+{
+	// Unsupported.
+	return false;
+}
+
+bool FSteamSocket::SetMulticastTtl(uint8 TimeToLive)
+{
+	// Unsupported.
+	return false;
+}
+
+bool FSteamSocket::SetMulticastInterface(const FInternetAddr& InterfaceAddress)
+{
+	// Unsupported.
+	return false;
+}
diff --git a/Engine/Plugins/Runtime/Steam/SteamSockets/Source/SteamSockets/Private/SteamSocket.h b/Engine/Plugins/Runtime/Steam/SteamSockets/Source/SteamSockets/Private/SteamSocket.h
new file mode 100644
--- /dev/null
+++ b/Engine/Plugins/Runtime/Steam/SteamSockets/Source/SteamSockets/Private/SteamSocket.h
@@ -0,0 +1,99 @@
+// Copyright 1998-2019 Epic Games, Inc. All Rights Reserved.
+
+#pragma once
+
+#include "CoreMinimal.h"
+#include "SteamSocketsPackage.h"
+#include "SteamSocketsPrivate.h"
+#include "IPAddressSteamSockets.h"
+#include "Sockets.h"
+#include "SocketTypes.h"
+
+class FSteamSocket : public FSocket
+{
+PACKAGE_SCOPE:
+	HSteamNetConnection InternalHandle;
+	// Binding address to determine the Steam Socket protocol we'll be using
+	FInternetAddrSteamSockets BindAddress;
+	int32 SendMode;
+	bool bShouldLingerOnClose;
+	// Flag to differentiate which calls can be performed on this socket
+	bool bIsListenSocket;
+	// Flag to make sure we set the weaker authentication (skips SDR validation) in a LAN environment.
+	bool bIsLANSocket;
+	// Easy flag to quickly determine if the pending data needs to be taken before expiration.
+	bool bHasPendingData;
+
+	// Steam Internal value used for shutting off a socket
+	ESteamNetConnectionEnd ClosureReason;
+
+	// This is used for peeking and looking at pending data
+	// Data is already internally handled in a queue on the SteamAPI side, this will give us whatever is at the top
+	// Only set if called with HasPendingData or Recv with a Peek flag.
+	SteamNetworkingMessage_t* PendingData;
+
+	class FSteamSocketsSubsystem* SocketSubsystem;
+
+	// Makes sure that sockets are properly set up in LAN environments.
+	void SetLanOptions();
+
+public:
+	// FSocket implementation
+	FSteamSocket(ESocketType InSocketType, const FString& InSocketDescription, const FName& InSocketProtocol);
+	virtual ~FSteamSocket();
+
+	virtual bool Close() override;
+
+	virtual bool Bind(const FInternetAddr& Addr) override;
+	virtual bool Connect(const FInternetAddr& Addr) override;
+	virtual bool Listen(int32 MaxBacklog) override;
+
+	virtual class FSocket* Accept(const FString& InSocketDescription) override;
+	virtual class FSocket* Accept(FInternetAddr& OutAddr, const FString& InSocketDescription) override { return nullptr; }
+
+	virtual bool SendTo(const uint8* Data, int32 Count, int32& BytesSent, const FInternetAddr& Destination) override;
+
+	virtual bool Send(const uint8* Data, int32 Count, int32& BytesSent) override;
+	virtual bool Recv(uint8* Data, int32 BufferSize, int32& BytesRead, ESocketReceiveFlags::Type Flags = ESocketReceiveFlags::None) override;
+
+	virtual bool HasPendingData(uint32& PendingDataSize) override;
+	virtual ESocketConnectionState GetConnectionState() override;
+
+	virtual void GetAddress(FInternetAddr& OutAddr) override { OutAddr = BindAddress; }
+	virtual bool GetPeerAddress(FInternetAddr& OutAddr) override;
+
+	virtual bool SetNoDelay(bool bIsNoDelay = true) override;
+	virtual bool SetLinger(bool bShouldLinger = true, int32 Timeout = 0) override;
+
+	virtual bool SetSendBufferSize(int32 Size, int32& NewSize) override;
+	virtual bool SetReceiveBufferSize(int32 Size, int32& NewSize) override;
+
+	virtual int32 GetPortNo() override
+	{
+		return BindAddress.GetPort();
+	}
+
+	/** Unsupported functions */
+	virtual bool Shutdown(ESocketShutdownMode Mode) override;
+	virtual bool Wait(ESocketWaitConditions::Type Condition, FTimespan WaitTime) override;
+	virtual bool WaitForPendingConnection(bool& bHasPendingConnection, const FTimespan& WaitTime) override;
+	virtual bool SetReuseAddr(bool bAllowReuse = true) override;
+	virtual bool SetRecvErr(bool bUseErrorQueue = true) override;
+	virtual bool SetNonBlocking(bool bIsNonBlocking = true) override;
+	virtual bool SetBroadcast(bool bAllowBroadcast = true) override;
+	virtual bool JoinMulticastGroup(const FInternetAddr& GroupAddress) override;
+	virtual bool JoinMulticastGroup(const FInternetAddr& GroupAddress, const FInternetAddr& InterfaceAddress) override;
+	virtual bool LeaveMulticastGroup(const FInternetAddr& GroupAddress) override;
+	virtual bool LeaveMulticastGroup(const FInternetAddr& GroupAddress, const FInternetAddr& InterfaceAddress);
+	virtual bool SetMulticastLoopback(bool bLoopback) override;
+	virtual bool SetMulticastTtl(uint8 TimeToLive) override;
+	virtual bool SetMulticastInterface(const FInternetAddr& InterfaceAddress) override;
+
+	// ~FSocket implementation
+
+	bool RecvRaw(SteamNetworkingMessage_t*& Data, int32 MaxMessages, int32& MessagesRead, ESocketReceiveFlags::Type Flags = ESocketReceiveFlags::None);
+
+	void SetClosureReason(ESteamNetConnectionEnd NewClosureReason) { ClosureReason = NewClosureReason; }
+	void SetSendMode(int32 NewSendMode);
+
+};
diff --git a/Engine/Plugins/Runtime/Steam/SteamSockets/Source/SteamSockets/Private/SteamSocketsModule.cpp b/Engine/Plugins/Runtime/Steam/SteamSockets/Source/SteamSockets/Private/SteamSocketsModule.cpp
new file mode 100644
--- /dev/null
+++ b/Engine/Plugins/Runtime/Steam/SteamSockets/Source/SteamSockets/Private/SteamSocketsModule.cpp
@@ -0,0 +1,74 @@
+// Copyright 1998-2019 Epic Games, Inc. All Rights Reserved.
+
+#include "SteamSocketsModule.h"
+#include "SteamSocketsSubsystem.h"
+#include "Misc/ConfigCacheIni.h"
+#include "Modules/ModuleManager.h"
+#include "OnlineSubsystemNames.h"
+#include "SteamSharedModule.h"
+#include "SteamSocketsTypes.h"
+#include "SocketSubsystemModule.h"
+
+IMPLEMENT_MODULE(FSteamSocketsModule, SteamSockets);
+
+namespace FNetworkProtocolTypes
+{
+	const FName SteamSocketsP2P(TEXT("SteamSocketsP2P"));
+	const FName SteamSocketsIP(TEXT("SteamSocketsIP"));
+}
+
+void FSteamSocketsModule::StartupModule()
+{
+	FSteamSharedModule& SharedModule = FSteamSharedModule::Get();
+	const bool bIsNotEditor = (IsRunningDedicatedServer() || IsRunningGame());
+
+	// Load the Steam modules before first call to API
+	if (SharedModule.AreSteamDllsLoaded() && bIsNotEditor)
+	{
+		// Settings flags
+		bool bOverrideSocketSubsystem = true;
+
+		// Use this flag from the SteamNetworking configuration so that the IPNetDrivers can be used as needed.
+		if (GConfig)
+		{
+			GConfig->GetBool(TEXT("OnlineSubsystemSteam"), TEXT("bUseSteamNetworking"), bOverrideSocketSubsystem, GEngineIni);
+		}
+
+		// Create and register our singleton factory with the main online subsystem for easy access
+		FSteamSocketsSubsystem* SocketSubsystem = FSteamSocketsSubsystem::Create();
+		FString Error;
+		if (SocketSubsystem->Init(Error))
+		{
+			bEnabled = true;
+
+			// Register our socket Subsystem
+			FSocketSubsystemModule& SSS = FModuleManager::LoadModuleChecked<FSocketSubsystemModule>("Sockets");
+			SSS.RegisterSocketSubsystem(STEAM_SOCKETS_SUBSYSTEM, SocketSubsystem, bOverrideSocketSubsystem);
+		}
+		else
+		{
+			UE_LOG(LogSockets, Error, TEXT("SteamSockets: Could not initialize SteamSockets, got error: %s"), *Error);
+			FSteamSocketsSubsystem::Destroy();
+		}
+	}
+	else if(bIsNotEditor)
+	{
+		UE_LOG(LogSockets, Warning, TEXT("SteamSockets: Steam SDK %s libraries not present at %s or failed to load!"), STEAM_SDK_VER, *SharedModule.GetSteamModulePath());
+	}
+	else
+	{
+		UE_LOG(LogSockets, Log, TEXT("SteamSockets: Disabled for editor process."));
+	}
+}
+
+void FSteamSocketsModule::ShutdownModule()
+{
+	FModuleManager& ModuleManager = FModuleManager::Get();
+
+	if (ModuleManager.IsModuleLoaded("Sockets"))
+	{
+		FSocketSubsystemModule& SSS = FModuleManager::GetModuleChecked<FSocketSubsystemModule>("Sockets");
+		SSS.UnregisterSocketSubsystem(STEAM_SOCKETS_SUBSYSTEM);
+	}
+	FSteamSocketsSubsystem::Destroy();
+}
diff --git a/Engine/Plugins/Runtime/Steam/SteamSockets/Source/SteamSockets/Private/SteamSocketsNetConnection.cpp b/Engine/Plugins/Runtime/Steam/SteamSockets/Source/SteamSockets/Private/SteamSocketsNetConnection.cpp
new file mode 100644
--- /dev/null
+++ b/Engine/Plugins/Runtime/Steam/SteamSockets/Source/SteamSockets/Private/SteamSocketsNetConnection.cpp
@@ -0,0 +1,224 @@
+// Copyright 1998-2019 Epic Games, Inc. All Rights Reserved.
+
+#include "SteamSocketsNetConnection.h"
+#include "SteamSocketsPrivate.h"
+#include "IPAddressSteamSockets.h"
+#include "SteamSocketsSubsystem.h"
+#include "SteamSocketsNetDriver.h"
+#include "SteamSocket.h"
+#include "Net/DataChannel.h"
+#include "CoreMinimal.h"
+#include "PacketHandlers/StatelessConnectHandlerComponent.h"
+#include "Templates/SharedPointer.h"
+
+// Steam has the max value as 512 * 1024 (assumed bytes), which is unusable to us because UNetConnection asserts on any value greater than 1024
+// Because of this, we cannot pass the value we get from Valve.
+//
+// As this value may change in a future Steam iteration without any notice, this is just set to 1024.
+static const int32 MAX_PACKET = 1024;
+
+void USteamSocketsNetConnection::CleanUp()
+{
+	if (GetDriver() != nullptr && GetDriver()->GetSocketSubsystem() != nullptr && ConnectionSocket != nullptr)
+	{
+		FSteamSocketsSubsystem* SocketSub = static_cast<FSteamSocketsSubsystem*>(GetDriver()->GetSocketSubsystem());
+		if (SocketSub != nullptr)
+		{
+			SocketSub->QueueRemoval(ConnectionSocket->InternalHandle);
+		}
+	}
+
+	UNetConnection::CleanUp();
+}
+
+void USteamSocketsNetConnection::InitBase(UNetDriver* InDriver, FSocket* InSocket, const FURL& InURL, EConnectionState InState, int32 InMaxPacket, int32 InPacketOverhead)
+{
+	// There is no packet overhead on this protocol, but since the UNetConnection will assert without a value greater than 0, we give it 1.
+	UNetConnection::InitBase(InDriver, InSocket, InURL, InState, 
+		((InMaxPacket == 0) ? MAX_PACKET : InMaxPacket),
+		((InPacketOverhead == 0) ? 1 : InPacketOverhead));
+
+	// Keep a reference to the socket.
+	ConnectionSocket = static_cast<FSteamSocket*>(InSocket);
+}
+
+void USteamSocketsNetConnection::InitRemoteConnection(UNetDriver* InDriver, FSocket* InSocket, const FURL& InURL, const FInternetAddr& InRemoteAddr, EConnectionState InState,
+	int32 InMaxPacket, int32 InPacketOverhead)
+{
+	// Call base so we set up all of our settings
+	InitBase(InDriver, InSocket, InURL, InState, InMaxPacket, InPacketOverhead);
+
+	// Grab the connection address from the URL
+	RemoteAddr = InRemoteAddr.Clone();
+
+	InitSendBuffer();
+	SetClientLoginState(EClientLoginState::LoggingIn);
+	SetExpectedClientLoginMsgType(NMT_Hello);
+}
+
+void USteamSocketsNetConnection::InitLocalConnection(UNetDriver* InDriver, FSocket* InSocket, const FURL& InURL, EConnectionState InState, int32 InMaxPacket, int32 InPacketOverhead)
+{
+	// Call base so we set up all of our settings
+	InitBase(InDriver, InSocket, InURL, InState, InMaxPacket, InPacketOverhead);
+
+	if (InDriver == nullptr || InDriver->GetSocketSubsystem() == nullptr)
+	{
+		UE_LOG(LogNet, Error, TEXT("SteamSockets: InitLocalConnection: Unable to get SocketSubsystem"));
+		Close();
+		return;
+	}
+
+	// Grab the connection address from the URL
+	RemoteAddr = InDriver->GetSocketSubsystem()->GetAddressFromString(InURL.Host);
+	if (!RemoteAddr.IsValid() || !RemoteAddr->IsValid())
+	{
+		UE_LOG(LogNet, Error, TEXT("SteamSockets: Cannot create connection address from NetConnection URL"));
+		Close();
+		return;
+	}
+	RemoteAddr->SetPort(InURL.Port);
+
+	InitSendBuffer();
+}
+
+void USteamSocketsNetConnection::LowLevelSend(void* Data, int32 CountBits, FOutPacketTraits& Traits)
+{
+	// Grab the socket subsystem now, we'll use it later.
+	FSteamSocketsSubsystem* SocketSub = nullptr;
+	if (GetDriver() != nullptr && GetDriver()->GetSocketSubsystem() != nullptr)
+	{
+		SocketSub = static_cast<FSteamSocketsSubsystem*>(GetDriver()->GetSocketSubsystem());
+	}
+
+	if (ConnectionSocket != nullptr)
+	{
+		const uint8* SendData = reinterpret_cast<const uint8*>(Data);
+
+		// Process any packet modifiers
+		if (Handler.IsValid() && !Handler->GetRawSend())
+		{
+			const ProcessedPacket ProcessedData = Handler->Outgoing(reinterpret_cast<uint8*>(Data), CountBits, Traits);
+
+			if (!ProcessedData.bError)
+			{
+				SendData = ProcessedData.Data;
+				CountBits = ProcessedData.CountBits;
+			}
+			else
+			{
+				CountBits = 0;
+			}
+		}
+
+		int32 BytesSent = 0;		
+		int32 BytesToSend = FMath::DivideAndRoundUp(CountBits, 8);
+
+		// Make sure that we're within the acceptable size range.
+		if (BytesToSend > k_cbMaxSteamNetworkingSocketsMessageSizeSend)
+		{
+			UE_LOG(LogNet, Error, TEXT("SteamSockets: LowLevelSend: Attempted to send %d bytes of data, which is over the %d maximum"), 
+				BytesToSend, k_cbMaxSteamNetworkingSocketsMessageSizeSend);
+			SocketSub->LastSocketError = SE_EMSGSIZE;
+			return;
+		}
+		
+		if (BytesToSend > 0)
+		{
+			if (!ConnectionSocket->Send(SendData, BytesToSend, BytesSent))
+			{
+				if(SocketSub != nullptr)
+				{
+					UE_LOG(LogNet, Warning, TEXT("SteamSockets: LowLevelSend: Could not send %d bytes of data got error %d"), BytesToSend, (int32)SocketSub->GetLastErrorCode());
+				}
+			}
+		}
+	}
+	else if (SocketSub)
+	{
+		UE_LOG(LogNet, Error, TEXT("SteamSockets: LowLevelSend: Attempted to send data over a netconnection that has no socket!"));
+		SocketSub->LastSocketError = SE_EPROCLIM;
+	}
+}
+
+FString USteamSocketsNetConnection::LowLevelGetRemoteAddress(bool bAppendPort)
+{
+	// If we have a remote addr set already (should always have this set)
+	if (RemoteAddr.IsValid() && RemoteAddr->IsValid())
+	{
+		return RemoteAddr->ToString(bAppendPort);
+	}
+
+	// Otherwise, attempt to pull it from the socket itself.
+	FInternetAddrSteamSockets AddressData;
+	if (ConnectionSocket != nullptr)
+	{
+		ConnectionSocket->GetPeerAddress(AddressData);
+		if (AddressData.IsValid())
+		{
+			return AddressData.ToString(bAppendPort);
+		}
+	}
+	return TEXT("Invalid");
+}
+
+FString USteamSocketsNetConnection::LowLevelDescribe()
+{
+	return FString::Printf(TEXT("Addr=%s"), *RemoteAddr->ToString(true));
+}
+
+void USteamSocketsNetConnection::HandleRecvMessage(void* InData, int32 SizeOfData, const FInternetAddr* InFormattedAddress)
+{
+	uint8* RecvData = reinterpret_cast<uint8*>(InData);
+	if (RecvData == nullptr || InFormattedAddress == nullptr)
+	{
+		return;
+	}
+
+	if (bInConnectionlessHandshake)
+	{
+		USteamSocketsNetDriver* SteamNetDriver = static_cast<USteamSocketsNetDriver*>(Driver);
+		if (SteamNetDriver && !SteamNetDriver->ArePacketHandlersDisabled() && SteamNetDriver->ConnectionlessHandler.IsValid() && SteamNetDriver->StatelessConnectComponent.IsValid())
+		{
+			bool bRestartedHandshake = false;
+			bool bPassedChallenge = false;
+
+			TSharedPtr<const FInternetAddr> SharedAddress = InFormattedAddress->Clone();
+			const ProcessedPacket RawPacket = SteamNetDriver->ConnectionlessHandler->IncomingConnectionless(SharedAddress, RecvData, SizeOfData);
+			TSharedPtr<StatelessConnectHandlerComponent> StatelessConnect = SteamNetDriver->StatelessConnectComponent.Pin();
+
+			if (!RawPacket.bError && StatelessConnect->HasPassedChallenge(SharedAddress, bRestartedHandshake) && !bRestartedHandshake)
+			{
+				// Set the initial packet sequence from the handshake data
+				if (StatelessConnect.IsValid())
+				{
+					int32 ServerSequence = 0;
+					int32 ClientSequence = 0;
+
+					StatelessConnect->GetChallengeSequence(ServerSequence, ClientSequence);
+
+					InitSequence(ClientSequence, ServerSequence);
+				}
+
+				// Start handshaking now
+				if (Handler.IsValid())
+				{
+					Handler->BeginHandshaking();
+				}
+
+				UE_LOG(LogNet, Log, TEXT("SteamSockets: Connectionless handshake complete"));
+				bInConnectionlessHandshake = false;
+				SizeOfData = FMath::DivideAndRoundUp(RawPacket.CountBits, 8);
+				if (SizeOfData > 0)
+				{
+					RecvData = RawPacket.Data;
+				}
+				else
+				{
+					return;
+				}
+			}
+		}
+	}
+
+	UNetConnection::ReceivedRawPacket(RecvData, SizeOfData);
+}
diff --git a/Engine/Plugins/Runtime/Steam/SteamSockets/Source/SteamSockets/Private/SteamSocketsNetDriver.cpp b/Engine/Plugins/Runtime/Steam/SteamSockets/Source/SteamSockets/Private/SteamSocketsNetDriver.cpp
new file mode 100644
--- /dev/null
+++ b/Engine/Plugins/Runtime/Steam/SteamSockets/Source/SteamSockets/Private/SteamSocketsNetDriver.cpp
@@ -0,0 +1,568 @@
+// Copyright 1998-2019 Epic Games, Inc. All Rights Reserved.
+
+#include "SteamSocketsNetDriver.h"
+#include "SteamSocketsPrivate.h"
+#include "SocketSubsystem.h"
+#include "SteamSocketsTypes.h"
+#include "SteamSocketsModule.h"
+#include "SteamSocket.h"
+#include "SteamSocketsNetConnection.h"
+#include "SteamSocketsSubsystem.h"
+#include "IPAddressSteamSockets.h"
+
+void USteamSocketsNetDriver::PostInitProperties()
+{
+	Super::PostInitProperties();
+}
+
+void USteamSocketsNetDriver::Shutdown()
+{
+	UE_LOG(LogSockets, Verbose, TEXT("SteamSockets: Shutdown called on netdriver"));
+	FSteamSocketsSubsystem* SocketSub = static_cast<FSteamSocketsSubsystem*>(GetSocketSubsystem());
+	if (SocketSub && Socket)
+	{
+		FSteamSocketsSubsystem::FSteamSocketInformation* SocketInfo = SocketSub->GetSocketInfo(Socket->InternalHandle);
+		// Make sure to remove any netdriver information as we are already shutting down.
+		if (SocketInfo)
+		{
+			SocketInfo->NetDriver = nullptr;
+			SocketInfo->MarkForDeletion();
+		}
+
+		Socket = nullptr;
+	}
+
+	Super::Shutdown();
+}
+
+bool USteamSocketsNetDriver::IsAvailable() const
+{
+	const FSteamSocketsSubsystem* SocketSub = static_cast<const FSteamSocketsSubsystem*>(ISocketSubsystem::Get(STEAM_SOCKETS_SUBSYSTEM));
+	// Check if we have client handles and we active.
+	if (SocketSub && SocketSub->IsSteamInitialized() && FSteamSocketsModule::Get().IsSteamSocketsEnabled())
+	{
+		return true;
+	}
+	return false;
+}
+
+bool USteamSocketsNetDriver::InitBase(bool bInitAsClient, FNetworkNotify* InNotify, const FURL& URL, bool bReuseAddressAndPort, FString& Error)
+{
+	// Initialize base class
+	if (!UNetDriver::InitBase(bInitAsClient, InNotify, URL, bReuseAddressAndPort, Error))
+	{
+		UE_LOG(LogNet, Warning, TEXT("SteamSockets: Failed to initialize the base UNetDriver"));
+		return false;
+	}
+
+	if (!IsAvailable())
+	{
+		UE_LOG(LogNet, Warning, TEXT("SteamSockets: SteamSockets are not enabled, cannot use netdriver!"));
+		return false;
+	}
+
+	FName SocketAddressType = NAME_None;
+	const bool bIsLAN = URL.HasOption(TEXT("bIsLanMatch"));
+	const bool bForcedIP = (bIsLAN || URL.HasOption(TEXT("bPassthrough")));
+	FSteamSocketsSubsystem* SocketSub = static_cast<FSteamSocketsSubsystem*>(GetSocketSubsystem());
+	if (!SocketSub)
+	{
+		UE_LOG(LogNet, Warning, TEXT("SteamSockets: Unable to get SocketSubsystem"));
+		return false;
+	}
+
+	// Check to see if we're going to be using a P2P address
+	const bool bIsUsingSteamAddrs = (SocketSub->IsUsingRelayNetwork());
+
+	// This case is very special to the Steam framework and should not be followed as a template for other protocols.
+	// On Steam Sockets we need to be logged in as a dedicated server if we are hosted or just using the relay network.
+	// Because of this, we need to delay our actions and inquire for a login with the Steam network.
+	bIsDelayedNetworkAccess = !bInitAsClient &&
+		bIsUsingSteamAddrs &&
+		!SocketSub->IsLoggedInToSteam() && !bForcedIP;
+
+	// Delayed access and clients that will be hosting using SteamIDs will skip over the address determination and just use the P2P network.
+	if (bIsDelayedNetworkAccess || (!bInitAsClient && !IsRunningDedicatedServer() && bIsUsingSteamAddrs && !bForcedIP))
+	{
+		SocketAddressType = FNetworkProtocolTypes::SteamSocketsP2P;
+	}
+	else if (bForcedIP) // Support LAN Matches over Steam Sockets
+	{
+		SocketAddressType = FNetworkProtocolTypes::SteamSocketsIP;
+	}
+	else
+	{
+		// This is the typical behavior of a connection in a NetDriver.
+		// Figure out what kind of socket we'll need by parsing the Host field.
+		TSharedPtr<FInternetAddr> SocketAddressHost = SocketSub->GetAddressFromString(URL.Host);
+		if (SocketAddressHost.IsValid() && SocketAddressHost->IsValid())
+		{
+			SocketAddressType = SocketAddressHost->GetProtocolType();
+		}
+	}
+
+	// Create the socket itself
+	FString SocketDescription = (bInitAsClient) ? TEXT("Unreal client (Steam Sockets)") : TEXT("Unreal server (Steam Sockets)");
+	Socket = static_cast<FSteamSocket*>(SocketSub->CreateSocket(FName(TEXT("SteamSocket")), SocketDescription, SocketAddressType));
+	Socket->SetNoDelay(true);
+	
+	// Set up LAN flags early so that we make sure that the handshake doesn't go south when we don't have
+	// a stable connection to the Steam backend.
+	if (bIsLAN)
+	{
+		Socket->bIsLANSocket = true;
+	}
+
+	// On Servers, there's a possibility that the SocketAddressType is None due to not having enough information to determine
+	// the user's intent. Because of this, the SocketSubsystem will have determined a correct protocol for us to use. We need
+	// to update our address type so that we can find a good binding address to use later.
+	if (SocketAddressType.IsNone())
+	{
+		SocketAddressType = Socket->GetProtocol();
+	}
+
+	// Delayed server login causes us to have a valid bind address much later than an usual process.
+	// Because of this, a normal flow should be skipped over and use a forced P2P address.
+	if (!bIsDelayedNetworkAccess)
+	{
+		// Set the binding address.
+		FSharedInternetAddrArray Addresses;
+		if (SocketSub->GetLocalAdapterAddresses(Addresses))
+		{
+			// Attempt to find a binding address that matches protocol of the connection destination
+			for (int32 i = 0; i < Addresses.Num(); ++i)
+			{
+				UE_LOG(LogNet, Verbose, TEXT("SteamSockets: Looking at binding address %s"), *Addresses[i]->ToString(true));
+				if (Addresses[i]->GetProtocolType() == SocketAddressType)
+				{
+					UE_LOG(LogNet, Verbose, TEXT("SteamSockets: Picked this binding address."));
+					LocalAddr = Addresses[i]->Clone();
+					break;
+				}
+			}
+		}
+
+		// Check for binding address issues.
+		if (!LocalAddr.IsValid())
+		{
+			UE_LOG(LogNet, Error, TEXT("SteamSockets: Could not determine the binding address!"));
+			return false;
+		}
+	}
+	else
+	{
+		// Create a blank address, this will get filled in later when we log in.
+		LocalAddr = SocketSub->CreateInternetAddr();
+	}
+
+	// Set the port
+	LocalAddr->SetPort(URL.Port);
+
+	// This should never fail.
+	if (Socket->Bind(*LocalAddr) == false)
+	{
+		// If this fails, we somehow created a binding address that we cannot use.
+		UE_LOG(LogNet, Error, TEXT("SteamSockets: Invalid binding address used to create socket"));
+		return false;
+	}
+
+	return true;
+}
+
+bool USteamSocketsNetDriver::InitConnect(FNetworkNotify* InNotify, const FURL& ConnectURL, FString& Error)
+{
+	if (!InitBase(true, InNotify, ConnectURL, false, Error))
+	{
+		UE_LOG(LogNet, Warning, TEXT("SteamSockets: InitConnect failed while setting up base information"));
+		return false;
+	}
+
+	// Create the NetConnection and channels
+	ServerConnection = NewObject<USteamSocketsNetConnection>(NetConnectionClass);
+	ServerConnection->InitLocalConnection(this, Socket, ConnectURL, USOCK_Pending);
+	CreateInitialClientChannels();
+
+	// Set the sequence flags if we have no packet handlers
+	if (ArePacketHandlersDisabled())
+	{
+		ServerConnection->InitSequence(4, 4);
+	}
+
+	// Grab the subsystem
+	FSteamSocketsSubsystem* SteamSubsystem = static_cast<FSteamSocketsSubsystem*>(GetSocketSubsystem());
+	check(SteamSubsystem);
+
+	// Figure out where we are going.
+	TSharedPtr<FInternetAddr> SocketAddressHost = SteamSubsystem->GetAddressFromString(ConnectURL.Host);
+	if (!SocketAddressHost.IsValid() || !SocketAddressHost->IsValid())
+	{
+		UE_LOG(LogNet, Error, TEXT("SteamSockets: Could not obtain the address to connect to, had input %s"), *ConnectURL.Host);
+		return false;
+	}
+
+	SocketAddressHost->SetPort(ConnectURL.Port);
+
+	// Attempt to connect to the server
+	if (!Socket->Connect(*SocketAddressHost))
+	{
+		UE_LOG(LogNet, Warning, TEXT("SteamSockets: Could not connect to address %s, got error code %d"), *SocketAddressHost->ToString(true), 
+			(int32)SteamSubsystem->GetLastErrorCode());
+		return false;
+	}
+
+	// Link connection information for later
+	SteamSubsystem->LinkNetDriver(Socket, this);
+
+	UE_LOG(LogNet, Log, TEXT("Game client on port %i, rate %i"), ConnectURL.Port, ServerConnection->CurrentNetSpeed);
+
+	return true;
+}
+
+bool USteamSocketsNetDriver::InitListen(FNetworkNotify* InNotify, FURL& LocalURL, bool bReuseAddressAndPort, FString& Error)
+{
+	if (!InitBase(false, InNotify, LocalURL, bReuseAddressAndPort, Error))
+	{
+		UE_LOG(LogNet, Warning, TEXT("SteamSockets: InitListen failed while setting up base information"));
+		return false;
+	}
+
+	InitConnectionlessHandler();
+
+	// Tell the SteamAPI to actually start listening.
+	// This will actually create the socket with the SDK
+	if (!bIsDelayedNetworkAccess && Socket->Listen(0) == false)
+	{
+		UE_LOG(LogNet, Warning, TEXT("SteamSockets: InitListen failed to start listening on the Steam Network."));
+		Socket->SetClosureReason(k_ESteamNetConnectionEnd_Misc_RelayConnectivity);
+
+		return false;
+	}
+
+	// Watch for connections that come in.
+	FSteamSocketsSubsystem* SteamSubsystem = static_cast<FSteamSocketsSubsystem*>(GetSocketSubsystem());
+	check(SteamSubsystem);
+
+	if (!bIsDelayedNetworkAccess)
+	{
+		SteamSubsystem->LinkNetDriver(Socket, this);
+	}
+	else
+	{
+		SteamSubsystem->AddDelayedListener(Socket, this);
+	}
+
+	UE_LOG(LogNet, Log, TEXT("%s started listening on %d"), *GetDescription(), LocalURL.Port);
+	UE_CLOG((!bIsDelayedNetworkAccess && !UE_BUILD_SHIPPING), LogNet, Log, TEXT("SteamSockets: Listening with handle %u"), Socket->InternalHandle);
+	return true;
+}
+
+void USteamSocketsNetDriver::TickDispatch(float DeltaTime)
+{
+	// Do this first to clean up dead connections so we don't waste time processing them.
+	UNetDriver::TickDispatch(DeltaTime);
+
+	bool bIsAServer = (ServerConnection == nullptr);
+	FSteamSocketsSubsystem* SteamSubsystem = static_cast<FSteamSocketsSubsystem*>(GetSocketSubsystem());
+
+	// Receive packets until we cannot any longer.
+	while (!bIsDelayedNetworkAccess)
+	{
+		// This could be placed out of the loop, however if we thread API event calls, this will need to be checked on every recv
+		if (Socket != nullptr)
+		{
+
+#if !UE_BUILD_SHIPPING
+			// This code causes peek cycles to run in order to test if the functionality works properly.
+			if (SteamSubsystem && SteamSubsystem->bShouldTestPeek)
+			{
+				uint32 PendingDataSize;
+				if (Socket->HasPendingData(PendingDataSize))
+				{
+					UE_LOG(LogNet, Verbose, TEXT("SteamSockets: handle %u has %u data pending on the socket"), Socket->InternalHandle, PendingDataSize);
+				}
+			}
+#endif
+			int32 BytesRead = 0;
+			SteamNetworkingMessage_t* Message;
+			// Attempt to grab a message using the Socket
+			if (Socket->RecvRaw(Message, 1, BytesRead) && BytesRead > 0 && Message != nullptr)
+			{
+				USteamSocketsNetConnection* ConnectionToHandleMessage = static_cast<USteamSocketsNetConnection*>((bIsAServer) ? 
+					FindClientConnectionForHandle(Message->m_conn) : ServerConnection);
+
+				// Grab sender information for the purposes of logging
+				FInternetAddrSteamSockets MessageSender(Message->m_sender);
+
+				// Set the P2P channel information if we're not over IP (which will already have the right data set)
+				if (MessageSender.GetProtocolType() != FNetworkProtocolTypes::SteamSocketsIP)
+				{
+					MessageSender.SetPort(Message->m_nChannel);
+				}
+
+				// Process the message for this connection.
+				if (ConnectionToHandleMessage != nullptr)
+				{
+					UE_LOG(LogNet, VeryVerbose, TEXT("SteamSockets: Recieved packet from %s with size %d"), *MessageSender.ToString(true), Message->GetSize());
+					ConnectionToHandleMessage->HandleRecvMessage(Message->m_pData, Message->GetSize(), &MessageSender);
+					Message->Release();
+					continue;
+				}
+				else
+				{
+					UE_LOG(LogNet, Warning, TEXT("SteamSockets: Could not find connection information for sender %s (handle: %u)"), *MessageSender.ToString(true), Message->m_conn);
+					Message->Release();
+					continue;
+				}
+			}
+			else if (BytesRead == 0)
+			{
+				// We have no more messages, leave.
+				UE_LOG(LogNet, VeryVerbose, TEXT("SteamSockets: Exhausted message, exiting loop."));
+				break;
+			}
+			else
+			{
+				// In theory, we should have no information what so ever about connections that failed.
+				// The sender should have no information in here and this should only happen if our handle is null.
+
+				// A later disconnection event message from the API should handle cleanup of this object.
+				UE_CLOG(SteamSubsystem, LogNet, Warning, TEXT("SteamSockets: Could not recv message, got error code %d"), SteamSubsystem->GetLastErrorCode());
+				break;
+			}
+		}
+		else
+		{
+			// Silently exit if the socket is invalid.
+			return;
+		}
+	}
+}
+
+void USteamSocketsNetDriver::LowLevelSend(TSharedPtr<const FInternetAddr> Address, void* Data, int32 CountBits, FOutPacketTraits& Traits)
+{
+	if (Address.IsValid() && Address->IsValid() && Socket != nullptr)
+	{
+		const uint8* SendData = reinterpret_cast<const uint8*>(Data);
+		if (ConnectionlessHandler.IsValid())
+		{
+			const ProcessedPacket ProcessedData =
+				ConnectionlessHandler->OutgoingConnectionless(Address, (uint8*)SendData, CountBits, Traits);
+
+			if (!ProcessedData.bError)
+			{
+				SendData = ProcessedData.Data;
+				CountBits = ProcessedData.CountBits;
+			}
+			else
+			{
+				CountBits = 0;
+			}
+		}
+
+		if (CountBits > 0)
+		{
+			int32 BytesToSend = FMath::DivideAndRoundUp(CountBits, 8);
+			int32 SentBytes = 0;
+
+			// Our sendto will find the correct socket to send over.
+			if (!Socket->SendTo(SendData, BytesToSend, SentBytes, *Address))
+			{
+				UE_LOG(LogNet, Warning, TEXT("SteamSockets: LowLevelSend: Could not send %d data over socket to %s!"), BytesToSend, *Address->ToString(true));
+			}
+		}
+	}
+	else
+	{
+		UE_LOG(LogNet, Warning, TEXT("SteamSockets: LowLevelSend: Address or Socket was invalid, could not send data."));
+	}
+}
+
+void USteamSocketsNetDriver::LowLevelDestroy()
+{
+	UNetDriver::LowLevelDestroy();
+
+	FSteamSocketsSubsystem* SteamSubsystem = static_cast<FSteamSocketsSubsystem*>(GetSocketSubsystem());
+	if (Socket != nullptr && !HasAnyFlags(RF_ClassDefaultObject))
+	{
+		SteamSubsystem->QueueRemoval(Socket->InternalHandle);
+		Socket = nullptr;
+
+		UE_LOG(LogNet, Log, TEXT("%s shut down"), *GetDescription());
+	}
+}
+
+class ISocketSubsystem* USteamSocketsNetDriver::GetSocketSubsystem()
+{
+	return ISocketSubsystem::Get(STEAM_SOCKETS_SUBSYSTEM);
+}
+
+bool USteamSocketsNetDriver::IsNetResourceValid(void)
+{
+	return Socket != nullptr && (!ServerConnection || (ServerConnection && ServerConnection->State == USOCK_Open));
+}
+
+bool USteamSocketsNetDriver::ArePacketHandlersDisabled() const
+{
+#if !UE_BUILD_SHIPPING
+	return FParse::Param(FCommandLine::Get(), TEXT("NoPacketHandler"));
+#else
+	return false;
+#endif
+}
+
+void USteamSocketsNetDriver::ResetSocketInfo(const FSteamSocket* RemovedSocket)
+{
+	const SteamSocketHandles SocketHandle = RemovedSocket->InternalHandle;
+	USteamSocketsNetConnection* SocketConnection = 
+		static_cast<USteamSocketsNetConnection*>(ServerConnection ? ServerConnection : FindClientConnectionForHandle(SocketHandle));
+
+	if (SocketConnection)
+	{
+		SocketConnection->ClearSocket();
+	}
+
+	// If this netdriver has the same socket pointer, go ahead and remove it.
+	if (Socket == RemovedSocket)
+	{
+		Socket = nullptr;
+	}
+}
+
+UNetConnection* USteamSocketsNetDriver::FindClientConnectionForHandle(SteamSocketHandles SocketHandle)
+{
+	for (UNetConnection*& ClientConnection : ClientConnections)
+	{
+		USteamSocketsNetConnection* SteamConnection = static_cast<USteamSocketsNetConnection*>(ClientConnection);
+		if (SteamConnection && SteamConnection->GetRawSocket() != nullptr)
+		{
+			const FSteamSocket* SteamSocket = SteamConnection->GetRawSocket();
+			if (SteamSocket && SteamSocket->InternalHandle == SocketHandle)
+			{
+				return ClientConnection;
+			}
+		}
+	}
+	return nullptr;
+}
+
+void USteamSocketsNetDriver::OnConnectionCreated(SteamSocketHandles ListenParentHandle, SteamSocketHandles SocketHandle)
+{
+	FSteamSocketsSubsystem* SocketSubsystem = static_cast<FSteamSocketsSubsystem*>(GetSocketSubsystem());
+
+	// If this netdriver does not hold the listener for the incoming connection, do not handle it
+	// This makes sure cleanup of listener socket later does not destroy sockets it does not own.
+	// (useful for beacons)
+	if (Socket == nullptr || ListenParentHandle != ((FSteamSocket*)Socket)->InternalHandle || !SocketSubsystem)
+	{
+		return;
+	}
+
+	// Absolutely make sure that we can interface with the SteamAPI
+	check(FSteamSocketsSubsystem::GetSteamSocketsInterface());
+
+	// Unlike the other SteamNetworking functionality, connections that we don't want cannot just be ignored
+	// So instead of processing everything and then disconnecting unwanted connections, we drop them immediately.
+	if (Notify->NotifyAcceptingConnection() == EAcceptConnection::Accept)
+	{
+		// Accept the connection with the API. We'll want to do this as quick as possible.
+		EResult AcceptedResult = FSteamSocketsSubsystem::GetSteamSocketsInterface()->AcceptConnection(SocketHandle);
+		if (AcceptedResult != k_EResultOK)
+		{
+			// We can fail here because the client aborted or something happened to our connection to the network
+			UE_LOG(LogNet, Error, TEXT("SteamSockets: Cannot accept connection due to error %d"), (int32)AcceptedResult);
+			SocketSubsystem->LastSocketError = SE_ECONNRESET;
+			return;
+		}
+
+		// Create the new connection
+		USteamSocketsNetConnection* NewConnection = NewObject<USteamSocketsNetConnection>(NetConnectionClass);
+		check(NewConnection);
+
+		// Create the bookkeeping and set up the socket.
+		FInternetAddrSteamSockets ConnectedAddr;
+		FSteamSocket* NewSocket = static_cast<FSteamSocket*>(Socket->Accept(TEXT("AcceptedSocket")));
+		NewSocket->InternalHandle = SocketHandle;
+		NewSocket->GetPeerAddress(ConnectedAddr);
+		SocketSubsystem->AddSocket(ConnectedAddr, NewSocket, Socket);
+		SocketSubsystem->LinkNetDriver(NewSocket, this);
+
+		EConnectionState RemoteConnectionState = USOCK_Pending;
+		SteamNetworkingQuickConnectionStatus ConnStatus;
+		// We might already have a full route to the user already, which will not signal an event later (for the listener) that would otherwise update the connection state
+		// Because of this, we need to quickly query the new connection status to determine where we are in the connection flow
+		// This is very unlikely unless we already had a previous route to this connection.
+		if (FSteamSocketsSubsystem::GetSteamSocketsInterface()->GetQuickConnectionStatus(SocketHandle, &ConnStatus) && ConnStatus.m_eState == k_ESteamNetworkingConnectionState_Connected)
+		{
+			// We have already found a route and made a connection, skip out of pending.
+			RemoteConnectionState = USOCK_Open;
+		}
+
+		// Create the net connection structures and add the connections.
+		NewConnection->InitRemoteConnection(this, NewSocket, World ? World->URL : FURL(), ConnectedAddr, RemoteConnectionState);
+
+		if (ArePacketHandlersDisabled())
+		{
+			// Set up the sequence numbers. Valve uses their own, but so do we so to prevent
+			// reinventing the wheel, we'll set up our sequence numbers to be some random garbage
+			NewConnection->InitSequence(4, 4);
+
+			// Attempt to start the PacketHandler handshakes (we do not support stateless connect)
+			if (NewConnection->Handler.IsValid())
+			{
+				NewConnection->Handler->BeginHandshaking();
+			}
+		}
+		else if (ConnectionlessHandler.IsValid() && StatelessConnectComponent.IsValid())
+		{
+			NewConnection->FlagForHandshake();
+		}
+
+		Notify->NotifyAcceptedConnection(NewConnection);
+		AddClientConnection(NewConnection);
+
+		UE_LOG(LogNet, Log, TEXT("SteamSockets: New connection (%u) over listening socket accepted %s"), SocketHandle, *ConnectedAddr.ToString(true));
+	}
+	else
+	{
+		FSteamSocketsSubsystem::GetSteamSocketsInterface()->CloseConnection(SocketHandle, k_ESteamNetConnectionEnd_App_Generic, "Connection rejected", false);
+		UE_LOG(LogNet, Log, TEXT("SteamSockets: New connection over listening socket rejected."));
+	}
+}
+
+void USteamSocketsNetDriver::OnConnectionUpdated(SteamSocketHandles SocketHandle, int32 NewState)
+{
+	// While this function can also get the relay status flags, we really only
+	// care about the connected state flag.
+	if (NewState == k_ESteamNetworkingConnectionState_Connected)
+	{
+		UNetConnection* SocketConnection = (ServerConnection) ? ServerConnection : FindClientConnectionForHandle(SocketHandle);
+		if (SocketConnection)
+		{
+			SocketConnection->State = USOCK_Open;
+		}
+
+		UE_LOG(LogNet, Verbose, TEXT("SteamSockets: Connection established with user with socket id: %u"), SocketHandle);
+
+		// For clients, we need to match the settings of the server
+		if (ServerConnection && ArePacketHandlersDisabled())
+		{
+			// Attempt to start the PacketHandler handshakes (we do not support stateless connect)
+			// The PendingNetGame will also do it but we can't actually send the packets for it until we're connected.
+			if (ServerConnection->Handler.IsValid())
+			{
+				ServerConnection->Handler->BeginHandshaking();
+			}
+		}
+	}
+}
+
+void USteamSocketsNetDriver::OnConnectionDisconnected(SteamSocketHandles SocketHandle)
+{
+	UNetConnection* SocketConnection = ServerConnection ? ServerConnection : FindClientConnectionForHandle(SocketHandle);
+	if (SocketConnection)
+	{
+		SocketConnection->State = USOCK_Closed;
+	}
+
+	UE_LOG(LogNet, Verbose, TEXT("SteamSockets: Connection dropped with user with socket id: %u"), SocketHandle);
+}
diff --git a/Engine/Plugins/Runtime/Steam/SteamSockets/Source/SteamSockets/Private/SteamSocketsPing.cpp b/Engine/Plugins/Runtime/Steam/SteamSockets/Source/SteamSockets/Private/SteamSocketsPing.cpp
new file mode 100644
--- /dev/null
+++ b/Engine/Plugins/Runtime/Steam/SteamSockets/Source/SteamSockets/Private/SteamSocketsPing.cpp
@@ -0,0 +1,98 @@
+// Copyright 1998-2019 Epic Games, Inc. All Rights Reserved.
+
+#include "SteamSocketsPing.h"
+#include "EngineLogs.h"
+#include "SteamSocketsPrivate.h"
+#include "SteamSocketsSubsystem.h"
+#include "Misc/CoreMisc.h"
+
+bool FSteamSocketsPing::IsUsingP2PRelays() const
+{
+	if (SocketSub)
+	{
+		return SocketSub->IsUsingRelayNetwork();
+	}
+
+	return false;
+}
+
+FString FSteamSocketsPing::GetHostPingData() const
+{
+	// Don't do anything if we're not using the relay network!
+	if (!IsUsingP2PRelays())
+	{
+		UE_LOG(LogNet, Warning, TEXT("SteamSockets: GetHostPingData was called but relays are disabled! This is invalid!"));
+		return TEXT("");
+	}
+
+	if (!SteamNetworkingUtils() || !SocketSub)
+	{
+		UE_LOG(LogNet, Warning, TEXT("SteamSockets: Could not get host ping data as SteamNetworking is not available."));
+		return TEXT("");
+	}
+
+	SteamNetworkPingLocation_t PingData;
+	// This will use whatever latest ping data we have, even if we're recalculating it.
+	// This should be near constant time.
+	float PingDataAge = SteamNetworkingUtils()->GetLocalPingLocation(PingData);
+	// While rare, this could happen if we request ping before all the interfaces are ready to go.
+	if (PingDataAge == -1)
+	{
+		// This is okay though, because the next call to this should have actual ping information.
+		UE_LOG(LogNet, Log, TEXT("SteamSockets: No ping data is available at this time"));
+		return TEXT("");
+	}
+
+	// Copy some data around so that we can send this data over the network
+	// (This must be done as data calculated on one machine is incompatible over the wire)
+	ANSICHAR ConversionBuffer[k_cchMaxSteamNetworkingPingLocationString];
+	FMemory::Memzero(ConversionBuffer); // Zero this for safety.
+	SteamNetworkingUtils()->ConvertPingLocationToString(PingData, ConversionBuffer, k_cchMaxSteamNetworkingPingLocationString);
+	
+	return FString(ANSI_TO_TCHAR(ConversionBuffer));
+}
+
+int32 FSteamSocketsPing::GetPingFromHostData(const FString& HostPingStr) const
+{
+	if (!SteamNetworkingUtils() || HostPingStr.IsEmpty() || !IsUsingP2PRelays())
+	{
+		UE_LOG(LogNet, Warning, TEXT("SteamSockets: Could not determine ping data as SteamNetworking is unavailable."));
+		return -1;
+	}
+
+	// Attempt to serialize the data that we've gotten already.
+	SteamNetworkPingLocation_t HostPingData;
+	if (SteamNetworkingUtils()->ParsePingLocationString(TCHAR_TO_ANSI(*HostPingStr), HostPingData))
+	{
+		// This should be an extremely quick function, as it compares regional data within the node network
+		// that we got from the host to our own node information. This ping data is mostly just math.
+		int32 PingValue = SteamNetworkingUtils()->EstimatePingTimeFromLocalHost(HostPingData);
+		if (PingValue == -1)
+		{
+			// A value of -1 occurs when there's a critical failure (We don't have our own ping data ready).
+			// However, it's not a fatal error, and we can always request for that data again later.
+			return -1;
+		}
+
+		return PingValue;
+	}
+	else
+	{
+		// This happens if the string has been tampered with, the protocol has upgraded or just something overall went wrong.
+		UE_LOG(LogNet, Log, TEXT("SteamSockets: Could not parse the host location in order to determine ping."));
+		return -1;
+	}
+}
+
+
+
+bool FSteamSocketsPing::IsRecalculatingPing() const
+{
+	if (!SteamNetworkingUtils() && IsUsingP2PRelays())
+	{
+		return false;
+	}
+
+	ESteamNetworkingAvailability NetworkStatus = SteamNetworkingUtils()->GetRelayNetworkStatus(nullptr);
+	return NetworkStatus == k_ESteamNetworkingAvailability_Retrying || NetworkStatus == k_ESteamNetworkingAvailability_Waiting || NetworkStatus == k_ESteamNetworkingAvailability_Attempting;
+}
diff --git a/Engine/Plugins/Runtime/Steam/SteamSockets/Source/SteamSockets/Private/SteamSocketsPing.h b/Engine/Plugins/Runtime/Steam/SteamSockets/Source/SteamSockets/Private/SteamSocketsPing.h
new file mode 100644
--- /dev/null
+++ b/Engine/Plugins/Runtime/Steam/SteamSockets/Source/SteamSockets/Private/SteamSocketsPing.h
@@ -0,0 +1,26 @@
+// Copyright 1998-2019 Epic Games, Inc. All Rights Reserved.
+
+#pragma once
+
+#include "OnlinePingInterfaceSteam.h"
+
+/**
+ * Steam Sockets ping interface implementation. See OnlinePingInterfaceSteam.h for information
+ */
+class FSteamSocketsPing : public FOnlinePingInterfaceSteam
+{
+public:
+	FSteamSocketsPing(class FSteamSocketsSubsystem* InSocketSub, class FOnlineSubsystemSteam* InOnlineSub) :
+		FOnlinePingInterfaceSteam(InOnlineSub),
+		SocketSub(InSocketSub)
+	{
+	}
+	
+	virtual bool IsUsingP2PRelays() const override;
+	virtual FString GetHostPingData() const override;
+	virtual int32 GetPingFromHostData(const FString& HostPingStr) const override;
+	virtual bool IsRecalculatingPing() const override;
+	
+protected:
+	class FSteamSocketsSubsystem* SocketSub;
+};
diff --git a/Engine/Plugins/Runtime/Steam/SteamSockets/Source/SteamSockets/Private/SteamSocketsPrivate.h b/Engine/Plugins/Runtime/Steam/SteamSockets/Source/SteamSockets/Private/SteamSocketsPrivate.h
new file mode 100644
--- /dev/null
+++ b/Engine/Plugins/Runtime/Steam/SteamSockets/Source/SteamSockets/Private/SteamSocketsPrivate.h
@@ -0,0 +1,24 @@
+// Copyright 1998-2019 Epic Games, Inc. All Rights Reserved.
+
+#pragma once
+
+#include "CoreMinimal.h"
+
+// To make less changes to the Steam OSS, we'll define this here
+// This will allow us to easily snap into the Steam OSS code
+#ifndef STEAM_URL_PREFIX
+#define STEAM_URL_PREFIX TEXT("steam.")
+#endif
+
+THIRD_PARTY_INCLUDES_START
+
+// Main headers
+#include "steam/steam_api.h"
+#include "steam/steam_gameserver.h"
+
+// Socket specific headers
+#include "steam/isteamnetworkingsockets.h"
+#include "steam/isteamnetworkingutils.h"
+#include "steam/steamnetworkingtypes.h"
+
+THIRD_PARTY_INCLUDES_END
diff --git a/Engine/Plugins/Runtime/Steam/SteamSockets/Source/SteamSockets/Private/SteamSocketsSubsystem.cpp b/Engine/Plugins/Runtime/Steam/SteamSockets/Source/SteamSockets/Private/SteamSocketsSubsystem.cpp
new file mode 100644
--- /dev/null
+++ b/Engine/Plugins/Runtime/Steam/SteamSockets/Source/SteamSockets/Private/SteamSocketsSubsystem.cpp
@@ -0,0 +1,901 @@
+// Copyright 1998-2019 Epic Games, Inc. All Rights Reserved.
+
+#include "SteamSocketsSubsystem.h"
+#include "Misc/CoreMisc.h"
+#include "Misc/ConfigCacheIni.h"
+#include "IPAddressSteamSockets.h"
+#include "SteamSocketsPrivate.h"
+#include "SteamSocketsTaskManager.h"
+#include "SteamSocketsNetDriver.h"
+#include "SteamSharedModule.h"
+#include "SteamSocket.h"
+#include "SteamSocketsPing.h"
+#include "OnlineSubsystem.h"
+#include "OnlineSubsystemSteam.h"
+#include "OnlineSubsystemNames.h"
+#include "SteamSocketsTypes.h"
+
+// Log Category for The API Debugger
+DEFINE_LOG_CATEGORY_STATIC(LogSteamSocketsAPI, Log, All);
+
+/** Message logging hook for debugging Steam Sockets. */
+void SteamSocketsDebugLogger(ESteamNetworkingSocketsDebugOutputType nType, const char *pszMsg)
+{
+	FString OutputType;
+	ELogVerbosity::Type Verboseness = ELogVerbosity::Verbose;
+	switch (nType)
+	{
+		default:
+		case k_ESteamNetworkingSocketsDebugOutputType_None:
+			OutputType = TEXT("None");
+			break;
+		case k_ESteamNetworkingSocketsDebugOutputType_Bug:
+			OutputType = TEXT("Bug");
+			Verboseness = ELogVerbosity::Error;
+			break;
+		case k_ESteamNetworkingSocketsDebugOutputType_Error:
+			OutputType = TEXT("Error");
+			Verboseness = ELogVerbosity::Error;
+			break;
+		case k_ESteamNetworkingSocketsDebugOutputType_Important:
+			OutputType = TEXT("Important");
+			Verboseness = ELogVerbosity::Warning;
+			break;
+		case k_ESteamNetworkingSocketsDebugOutputType_Warning:
+			OutputType = TEXT("Warning");
+			Verboseness = ELogVerbosity::Warning;
+			break;
+		case k_ESteamNetworkingSocketsDebugOutputType_Everything:
+		case k_ESteamNetworkingSocketsDebugOutputType_Verbose:
+		case k_ESteamNetworkingSocketsDebugOutputType_Msg:
+			OutputType = TEXT("Log");
+			break;
+		case k_ESteamNetworkingSocketsDebugOutputType_Debug:
+			OutputType = TEXT("Debug");
+			break;
+	}
+
+	// Runtime detection of logging level and verbosity.
+	GLog->Log(LogSteamSocketsAPI.GetCategoryName(), Verboseness, FString::Printf(TEXT("SteamSockets API: %s %s"), *OutputType, ANSI_TO_TCHAR(pszMsg)));
+}
+
+FSteamSocketsSubsystem* FSteamSocketsSubsystem::SocketSingleton = nullptr;
+
+FSteamSocketsSubsystem* FSteamSocketsSubsystem::Create()
+{
+	if (SocketSingleton == nullptr)
+	{
+		SocketSingleton = new FSteamSocketsSubsystem();
+	}
+
+	return SocketSingleton;
+}
+
+void FSteamSocketsSubsystem::Destroy()
+{
+	if (SocketSingleton != nullptr)
+	{
+		SocketSingleton->Shutdown();
+		delete SocketSingleton;
+		SocketSingleton = nullptr;
+	}
+}
+
+bool FSteamSocketsSubsystem::Init(FString& Error)
+{
+	const bool bIsDedicated = IsRunningDedicatedServer();
+	if (GConfig)
+	{
+		if (!GConfig->GetBool(TEXT("OnlineSubsystemSteam"), TEXT("bAllowP2PPacketRelay"), bUseRelays, GEngineIni))
+		{
+			UE_CLOG(bIsDedicated, LogSockets, Warning, TEXT("SteamSockets: Missing config value for bAllowP2PPacketRelay, will be unable to determine ping functionality"));
+		}
+	}
+
+	// Debug commandline support for socket relays. 
+#if !UE_BUILD_SHIPPING
+	bool bOverrideRelays = false;
+	if (FParse::Bool(FCommandLine::Get(), TEXT("SteamSocketsRelays"), bOverrideRelays))
+	{
+		UE_LOG(LogSockets, Log, TEXT("SteamSockets: Set relay setting to %d"), bOverrideRelays);
+		bUseRelays = bOverrideRelays;
+	}
+#endif
+
+	// Get the online subsystem
+	FOnlineSubsystemSteam* OnlineSteamSubsystem = static_cast<FOnlineSubsystemSteam*>(IOnlineSubsystem::Get(STEAM_SUBSYSTEM));
+
+	// Initialize the Server API if we are going to be using it.
+	if (bIsDedicated)
+	{
+		SteamAPIServerHandle = FSteamSharedModule::Get().ObtainSteamServerInstanceHandle();
+		if (!SteamAPIServerHandle.IsValid())
+		{
+			Error = TEXT("SteamSockets: Could not initialize the game server SteamAPI!");
+			SteamAPIServerHandle.Reset();
+			return false;
+		}
+
+		// Register the GameServer login status call for when the session system logs us in.
+		if (OnlineSteamSubsystem != nullptr)
+		{
+			SteamServerLoginDelegateHandle = OnlineSteamSubsystem->AddOnSteamServerLoginCompletedDelegate_Handle(
+				FOnSteamServerLoginCompletedDelegate::CreateRaw(this, &FSteamSocketsSubsystem::OnServerLoginComplete));
+		}
+		else
+		{
+			Error = TEXT("SteamSockets: Cannot register dedicated server login listeners!");
+			SteamAPIServerHandle.Reset();
+			return false;
+		}
+	}
+	else
+	{
+		SteamAPIClientHandle = FSteamSharedModule::Get().ObtainSteamClientInstanceHandle();
+		if (!SteamAPIClientHandle.IsValid())
+		{
+			Error = TEXT("SteamSockets: could not obtain a handle to SteamAPI!");
+			SteamAPIClientHandle.Reset();
+			return false;
+		}
+	}
+
+	// Clients and servers using the relay network will want to set this up
+	if (IsUsingRelayNetwork())
+	{
+		// We need this functionality in order to modify socket configuration and determine ping data.
+		if (!SteamNetworkingUtils())
+		{
+			Error = TEXT("SteamSockets: Cannot interface with SteamNetworkingUtils! Relay support relies on this functionality!");
+			SteamAPIClientHandle.Reset();
+			SteamAPIServerHandle.Reset();
+			return false;
+		}
+
+		// We do not care about the return of this function the first time it runs.
+		UE_LOG(LogSockets, Log, TEXT("SteamSockets: Initializing Network Relay"));
+		SteamNetworkingUtils()->InitRelayNetworkAccess();
+	}
+
+	// Start up the authentication system so we get the signing certs for this application session.
+	if (GetSteamSocketsInterface())
+	{
+		GetSteamSocketsInterface()->InitAuthentication();
+	}
+	
+	// Set up the API Logging
+	ESteamNetworkingSocketsDebugOutputType DebugLevel = k_ESteamNetworkingSocketsDebugOutputType_Msg;
+
+	// Set this a bit higher for shipping builds because SteamSockets can spew a lot of information
+#if UE_BUILD_SHIPPING
+	DebugLevel = k_ESteamNetworkingSocketsDebugOutputType_Important;
+#endif
+
+	SteamNetworkingUtils()->SetDebugOutputFunction(DebugLevel, SteamSocketsDebugLogger);
+
+	// Hook up the Steam Event Manager so we can get connection updates.
+	SteamEventManager = MakeUnique<FSteamSocketsTaskManager>(this);
+
+	// Set the ping interface up, this will allow the Steam OSS to compile without direct links to this module.
+	if (OnlineSteamSubsystem != nullptr && IsUsingRelayNetwork())
+	{
+		// Deletion will be handled by the Steam OSS, we do not need to worry about it in this subsystem.
+		TSharedPtr<FSteamSocketsPing, ESPMode::ThreadSafe> SteamPingHandler = MakeShareable(new FSteamSocketsPing(this, OnlineSteamSubsystem));
+		FOnlinePingSteamPtr NewPingInterface = StaticCastSharedPtr<FOnlinePingInterfaceSteam>(SteamPingHandler);
+		OnlineSteamSubsystem->SetPingInterface(NewPingInterface);
+	}
+
+	return true;
+}
+
+void FSteamSocketsSubsystem::Shutdown()
+{
+	UE_LOG(LogSockets, Log, TEXT("SteamSockets: Cleaning up"));
+
+	// Clear our delegate handles
+	FOnlineSubsystemSteam* OnlineSteamSubsystem = static_cast<FOnlineSubsystemSteam*>(IOnlineSubsystem::Get(STEAM_SUBSYSTEM));
+	if (OnlineSteamSubsystem)
+	{
+		OnlineSteamSubsystem->ClearOnSteamServerLoginCompletedDelegate_Handle(SteamServerLoginDelegateHandle);
+	}
+
+	// Clean up the internal event manager.
+	if (SteamEventManager.IsValid())
+	{
+		SteamEventManager.Reset();
+	}
+
+	// Remove all our listeners that were pending.
+	PendingListenerArray.Empty();
+
+	// Let go of our API handles. This is quicker than waiting for destruction
+	SteamAPIClientHandle.Reset();
+	SteamAPIServerHandle.Reset();
+}
+
+class FSocket* FSteamSocketsSubsystem::CreateSocket(const FName& SocketType, const FString& SocketDescription, const FName& ProtocolType)
+{
+	FName ProtocolTypeToUse = ProtocolType;
+	if (ProtocolTypeToUse.IsNone())
+	{
+		if (IsUsingRelayNetwork())
+		{
+			ProtocolTypeToUse = FNetworkProtocolTypes::SteamSocketsP2P;
+		}
+		else
+		{
+			// In addition to the P2P framework, the Steam Sockets can also work over IP.
+			// Generally, it's better to use the UIpNetDriver, but to prevent us from crashing
+			// if you have relays disabled, the fallback is the IP framework.
+			ProtocolTypeToUse = FNetworkProtocolTypes::SteamSocketsIP;
+		}
+	}
+	else if (ProtocolType != FNetworkProtocolTypes::SteamSocketsP2P && ProtocolType != FNetworkProtocolTypes::SteamSocketsIP)
+	{
+		UE_LOG(LogSockets, Warning, TEXT("SteamSockets: Attempted to create a socket with a protocol outside of the steam network!"));
+		return nullptr;
+	}
+
+	return static_cast<FSocket*>(new FSteamSocket(SOCKTYPE_Streaming, SocketDescription, ProtocolTypeToUse));
+}
+
+void FSteamSocketsSubsystem::DestroySocket(FSocket* Socket)
+{
+	if (Socket != nullptr)
+	{
+		FSteamSocket* SteamSocket = static_cast<FSteamSocket*>(Socket);
+		UE_LOG(LogSockets, Verbose, TEXT("SteamSockets: Destroy socket called on %u"), SteamSocket->InternalHandle);
+
+		// Save the internal handle because it will be invalid soon.
+		FSteamSocketInformation* SocketInfo = GetSocketInfo(SteamSocket->InternalHandle);
+		if (SocketInfo != nullptr)
+		{
+			// Warn the NetDriver that we're destructing the socket such that nothing else tampers with it.
+			USteamSocketsNetDriver* SteamNetDriver = SocketInfo->NetDriver.Get();
+			if (SteamNetDriver)
+			{
+				SteamNetDriver->ResetSocketInfo(SteamSocket);
+			}
+			SocketInfo->Socket = nullptr;
+			SocketInfo->MarkForDeletion(); // Mark us for deletion so we're removed from the map
+		}
+
+		// Socket closure will properly mark the socket for pending removal.
+		Socket->Close();
+
+		delete Socket;
+	}
+}
+
+FAddressInfoResult FSteamSocketsSubsystem::GetAddressInfo(const TCHAR* HostName, const TCHAR* ServiceName, 
+	EAddressInfoFlags QueryFlags, const FName ProtocolTypeName, ESocketType SocketType)
+{
+	// We don't support queries but if we can make an address from it then that's good enough.
+	FAddressInfoResult ResultData(HostName, ServiceName);
+
+	// Determine if we're asking for any specific protocol.
+	const bool bWildcardProtocol = ProtocolTypeName.IsNone();
+
+	// If you requested anything outside of the steam network, we will be unable to support it.
+	// This occurs when you have:
+	// * Given no hostname and service name information
+	// * Given a service name but the service name is not numeric (SteamNetwork can only take port numbers)
+	// * Specified a protocol type that's not a part of the steam network.
+	if ((HostName == nullptr && ServiceName == nullptr) || (ServiceName != nullptr && !FCString::IsNumeric(ServiceName)) ||
+		(ProtocolTypeName != FNetworkProtocolTypes::SteamSocketsP2P && ProtocolTypeName != FNetworkProtocolTypes::SteamSocketsIP && !bWildcardProtocol))
+	{
+		ResultData.ReturnCode = SE_EINVAL;
+		return ResultData;
+	}
+
+	FSharedInternetAddrArray AddressList;
+	// Handle binding flags or if you don't pass in a hostname
+	if (HostName == nullptr || EnumHasAnyFlags(QueryFlags, EAddressInfoFlags::BindableAddress))
+	{
+		// Grab all the adapters, we'll process them later.
+		// We don't care about the return type as on SteamSockets this call cannot fail.
+		GetLocalAdapterAddresses(AddressList);
+	}
+	else
+	{
+		TSharedPtr<FInternetAddr> SerializedAddr = GetAddressFromString(HostName);
+		// Only add to this list if the serialization succeeded.
+		if (SerializedAddr.IsValid())
+		{
+			AddressList.Add(SerializedAddr);
+		}
+	}
+	
+	if (AddressList.Num() >= 1)
+	{
+		// Assume we have no error, we'll check for this later.
+		ResultData.ReturnCode = SE_NO_ERROR;
+
+		// Preprocess the port information so we can set it later.
+		int32 PortToUse = -1;
+		if (ServiceName != nullptr && FCString::IsNumeric(ServiceName))
+		{
+			PortToUse = FCString::Atoi(ServiceName);
+		}
+
+		// Process the address list.
+		for (auto& AddressItem : AddressList)
+		{
+			// Check if the types match (add them anyways if protocol type wasn't specified).
+			if (bWildcardProtocol || (!bWildcardProtocol && AddressItem->GetProtocolType() == ProtocolTypeName))
+			{
+				// We will be writing directly into the array but this doesn't matter in our case as all the addresses are shared allocations.
+				if (PortToUse >= 0)
+				{
+					AddressItem->SetPort(PortToUse);
+				}
+				ResultData.Results.Add(FAddressInfoResultData(AddressItem.ToSharedRef(), 0, AddressItem->GetProtocolType(), SOCKTYPE_Streaming));
+			}
+		}
+
+		// If our address list is essentially empty due to mismatches, then set the correct error flag.
+		if (ResultData.Results.Num() <= 0)
+		{
+			if (!bWildcardProtocol)
+			{
+				// We couldn't find addresses that gave what the user was looking for.
+				ResultData.ReturnCode = SE_ADDRFAMILY;
+			}
+			else
+			{
+				// Otherwise something unforeseen happened.
+				ResultData.ReturnCode = SE_EFAULT;
+			}
+		}
+	}
+	else
+	{
+		ResultData.ReturnCode = SE_NO_DATA;
+	}
+
+	return ResultData;
+}
+
+TSharedPtr<FInternetAddr> FSteamSocketsSubsystem::GetAddressFromString(const FString& IPAddress)
+{
+	TSharedRef<FInternetAddrSteamSockets> NewAddr = StaticCastSharedRef<FInternetAddrSteamSockets>(CreateInternetAddr());
+
+	// Passing an empty string should just map it to the any address.
+	if (IPAddress.IsEmpty())
+	{
+		// NOTE: There's a lot of questions as to what this address should be here.
+		// Should we be returning the identity or the any address, 
+		// there's not really a good way to determine default.
+		NewAddr->SetAnyAddress();
+		
+		return NewAddr;
+	}
+
+	bool bIsAddrValid = false;
+	NewAddr->SetIp(*IPAddress, bIsAddrValid);
+	if (!bIsAddrValid)
+	{
+		return nullptr;
+	}
+
+	return NewAddr;
+}
+
+bool FSteamSocketsSubsystem::GetHostName(FString& HostName)
+{
+	// We have no way of looking up addresses on this platform.
+	UE_LOG(LogSockets, Warning, TEXT("GetHostName is not supported on SteamSockets"));
+	return false;
+}
+
+TSharedRef<FInternetAddr> FSteamSocketsSubsystem::CreateInternetAddr()
+{
+	return MakeShareable(new FInternetAddrSteamSockets());
+}
+
+const TCHAR* FSteamSocketsSubsystem::GetSocketAPIName() const
+{
+	return TEXT("SteamSockets");
+}
+
+bool FSteamSocketsSubsystem::GetLocalAdapterAddresses(FSharedInternetAddrArray& OutAddresses)
+{
+	// Add Multihome
+	TSharedRef<FInternetAddr> MultihomeAddress = CreateInternetAddr();
+	if (GetMultihomeAddress(MultihomeAddress))
+	{
+		OutAddresses.Add(MultihomeAddress);
+	}
+
+	// Add the account addresses
+	if (IsUsingRelayNetwork())
+	{
+		TSharedPtr<FInternetAddr> IdentityAddress = GetIdentityAddress();
+		if (IdentityAddress.IsValid())
+		{
+			OutAddresses.Add(IdentityAddress);
+		}
+	}
+
+	// Always include the any address.
+	TSharedRef<FInternetAddr> AnyAddress = CreateInternetAddr();
+	AnyAddress->SetAnyAddress();
+	OutAddresses.Add(AnyAddress);
+
+	return true;
+}
+
+TSharedRef<FInternetAddr> FSteamSocketsSubsystem::GetLocalBindAddr(FOutputDevice& Out)
+{
+	FSharedInternetAddrArray Adapters;
+	if (GetLocalAdapterAddresses(Adapters))
+	{
+		return Adapters[0].ToSharedRef();
+	}
+
+	return CreateInternetAddr();
+}
+
+void FSteamSocketsSubsystem::CleanSocketInformation(bool bForceClean)
+{
+	for (SocketHandleInfoMap::TIterator It(SocketInformationMap); It; ++It)
+	{
+		FSteamSocketInformation& SocketInfo = It.Value();
+		if (SocketInfo.IsMarkedForDeletion() || bForceClean)
+		{
+			// Close the netdriver and any connections.
+			USteamSocketsNetDriver* NetDriverInfo = SocketInfo.NetDriver.Get();
+			if (NetDriverInfo && SocketInfo.Parent == nullptr)
+			{
+				NetDriverInfo->Shutdown();
+				SocketInfo.NetDriver.Reset();
+			}
+
+			// Close the socket (frees internal memory)
+			//
+			// Deletion with SteamSockets is fairly safe. Multiple paths can mark a socket 
+			// for deletion but it will only ever be deleted once.
+			if (SocketInfo.Socket)
+			{
+				DestroySocket(SocketInfo.Socket);
+			}
+
+			It.RemoveCurrent();
+		}
+	}
+}
+
+void FSteamSocketsSubsystem::DumpSocketInformationMap() const
+{
+#if !UE_BUILD_SHIPPING
+	if (SocketInformationMap.Num() < 1)
+	{
+		UE_LOG(LogSockets, Log, TEXT("SteamSockets: Socket Information Map is empty!"));
+		return;
+	}
+
+	UE_LOG(LogSockets, Log, TEXT("SteamSockets: Printing Socket Information Map:\n"));
+	for (SocketHandleInfoMap::TConstIterator It(SocketInformationMap); It; ++It)
+	{
+		const FSteamSocketInformation& SocketInfo = It.Value();
+		UE_LOG(LogSockets, Log, TEXT("# %s"), *SocketInfo.ToString());
+	}
+#endif
+}
+
+bool FSteamSocketsSubsystem::Tick(float DeltaTime)
+{
+	// Handle all of our updates from the Steam API callbacks
+	if (SteamEventManager.IsValid())
+	{
+		SteamEventManager->Tick();
+	}
+
+	CleanSocketInformation(false);
+	return true;
+}
+
+bool FSteamSocketsSubsystem::Exec(class UWorld* InWorld, const TCHAR* Cmd, FOutputDevice& Ar)
+{
+	bool bIsHandled = false;
+
+#if !UE_BUILD_SHIPPING
+	if (FParse::Command(&Cmd, TEXT("PrintSteamSocketInfo")))
+	{
+		bIsHandled = true;
+		DumpSocketInformationMap();
+	}
+	if (FParse::Command(&Cmd, TEXT("PrintPendingSteamSocketInfo")))
+	{
+		bIsHandled = true;
+		UE_LOG(LogSockets, Log, TEXT("SteamSockets: Printing Pending Socket Info"));
+		for (auto& Pending : PendingListenerArray)
+		{
+			UE_LOG(LogSockets, Log, TEXT("# %s"), *Pending.ToString());
+		}
+	}
+	else if (FParse::Command(&Cmd, TEXT("ClearSteamSocketInfo")))
+	{
+		UE_LOG(LogSockets, Log, TEXT("SteamSockets: Clearing all socket information!"));
+		bIsHandled = true;
+		CleanSocketInformation(true);
+	}
+	else if (FParse::Command(&Cmd, TEXT("TogglePeekMessaging")))
+	{
+		bIsHandled = true;
+		bShouldTestPeek = !bShouldTestPeek;
+		UE_LOG(LogSockets, Log, TEXT("SteamSockets: Set Peek Messaging to %d"), bShouldTestPeek);
+	}
+	else if (FParse::Command(&Cmd, TEXT("TestSteamRawAddresses")))
+	{
+		UE_LOG(LogSockets, Log, TEXT("SteamSockets: Testing Raw addresses..."));
+		bIsHandled = true;
+		TSharedRef<FInternetAddrSteamSockets> SteamAddr = StaticCastSharedRef<FInternetAddrSteamSockets>(GetLocalBindAddr(*GLog));
+		TArray<uint8> RawAddress = SteamAddr->GetRawIp();
+		UE_LOG(LogSockets, Log, TEXT("SteamSockets: Address before %s"), *SteamAddr->ToString(false));
+		if (RawAddress.Num() <= 0)
+		{
+			UE_LOG(LogSockets, Warning, TEXT("SteamSockets: Raw address array is empty!"));
+			return true;
+		}
+		SteamAddr->SetRawIp(RawAddress);
+		UE_LOG(LogSockets, Log, TEXT("SteamSockets: Address after %s"), *SteamAddr->ToString(false));
+
+		TSharedPtr<FInternetAddrSteamSockets> SteamIPAddr = StaticCastSharedPtr<FInternetAddrSteamSockets>(GetAddressFromString(TEXT("8.8.8.8")));
+		if (SteamIPAddr)
+		{
+			TArray<uint8> RawIPAddress = SteamIPAddr->GetRawIp();
+			UE_LOG(LogSockets, Log, TEXT("SteamSockets: Address before %s"), *SteamIPAddr->ToString(false));
+			if (RawAddress.Num() <= 0)
+			{
+				UE_LOG(LogSockets, Warning, TEXT("SteamSockets: Raw address array is empty!"));
+				return true;
+			}
+			SteamIPAddr->SetRawIp(RawIPAddress);
+			UE_LOG(LogSockets, Log, TEXT("SteamSockets: Address after %s"), *SteamIPAddr->ToString(false));
+		}
+		UE_LOG(LogSockets, Log, TEXT("SteamSockets: Test complete"));
+	}
+#endif
+
+	return bIsHandled;
+}
+
+ISteamNetworkingSockets* FSteamSocketsSubsystem::GetSteamSocketsInterface()
+{
+	return (SteamGameServerNetworkingSockets() != nullptr && IsRunningDedicatedServer()) ? SteamGameServerNetworkingSockets() : SteamNetworkingSockets();
+}
+
+bool FSteamSocketsSubsystem::FSteamSocketInformation::operator==(const FInternetAddr& InAddr) const
+{
+	FInternetAddrSteamSockets SteamAddr = *((FInternetAddrSteamSockets*)&InAddr);
+	return *Addr == SteamAddr;
+}
+
+void FSteamSocketsSubsystem::FSteamSocketInformation::MarkForDeletion()
+{
+	bMarkedForDeletion = true;
+	USteamSocketsNetDriver* NetDriverObj = NetDriver.Get();
+	if (Socket != nullptr && NetDriverObj)
+	{
+		UE_LOG(LogSockets, Verbose, TEXT("SteamSockets: Handle %u marked for deletion. Has parent? %d NetDriver definition: %s"), 
+			Socket->InternalHandle, (Parent != nullptr), *NetDriverObj->GetDescription());
+	}
+}
+
+FString FSteamSocketsSubsystem::FSteamSocketInformation::ToString() const
+{
+	return FString::Printf(TEXT("SocketInfo: Addr[%s], Socket[%u], Status[%d], Listener[%d], HasNetDriver[%d], MarkedForDeletion[%d]"), 
+		(Addr.IsValid() ? *Addr->ToString(true) : TEXT("INVALID")), (Socket != nullptr ? Socket->InternalHandle : 0), 
+		(Socket != nullptr ? (int32)Socket->GetConnectionState() : -1),
+		(Parent != nullptr), 
+		NetDriver.IsValid(), bMarkedForDeletion);
+}
+
+FString FSteamSocketsSubsystem::FSteamPendingSocketInformation::ToString() const
+{
+	USteamSocketsNetDriver* SteamNetDriver = NetDriver.Get();
+	return FString::Printf(TEXT("PendingSocketInfo: NetDriver name %s"), (SteamNetDriver != nullptr) ? *SteamNetDriver->GetDescription() : TEXT("INVALID"));
+}
+
+FSteamSocketsSubsystem::FSteamSocketInformation* FSteamSocketsSubsystem::GetSocketInfo(SteamSocketHandles InternalSocketHandle)
+{
+	if (InternalSocketHandle == k_HSteamNetConnection_Invalid)
+	{
+		UE_LOG(LogSockets, Warning, TEXT("SteamSockets: Cannot get information on an invalid socket handle, returning null"));
+		return nullptr;
+	}
+
+	return SocketInformationMap.Find(InternalSocketHandle);
+}
+
+FSteamSocketsSubsystem::FSteamSocketInformation* FSteamSocketsSubsystem::GetSocketInfo(const FInternetAddr& ForAddress)
+{
+	// Small time save with invalid data
+	if (!ForAddress.IsValid())
+	{
+		return nullptr;
+	}
+
+	// This functionality is a bit slower, because we have to do complete address lookups through the map.
+	// It's much better to use the handles to find the socket info. There's also a possibility for collisions
+	// here as well. This is why internally, this function is not used as much (only when we don't have enough data).
+	for (SocketHandleInfoMap::TIterator It(SocketInformationMap); It; ++It)
+	{
+		if (It.Value() == ForAddress)
+		{
+			return &(It.Value());
+		}
+	}
+
+	return nullptr;
+}
+
+void FSteamSocketsSubsystem::AddSocket(const FInternetAddr& ForAddr, FSteamSocket* NewSocket, FSteamSocket* ParentSocket)
+{
+	if (NewSocket == nullptr || ForAddr.IsValid() == false)
+	{
+		UE_LOG(LogSockets, Warning, TEXT("SteamSockets: Attempted to track invalid socket data!"));
+		return;
+	}
+
+	// Do not attempt to add socket data for sockets that have invalid handles.
+	if (NewSocket->InternalHandle == k_HSteamListenSocket_Invalid || NewSocket->InternalHandle == k_HSteamNetConnection_Invalid)
+	{
+		UE_LOG(LogSockets, Verbose, TEXT("SteamSockets: Dropped socket tracking for socket with invalid handle."));
+		return;
+	}
+
+	FSteamSocketInformation NewSocketInfo(ForAddr.Clone(), NewSocket, ParentSocket);
+	if (SocketInformationMap.Find(NewSocket->InternalHandle) == nullptr)
+	{
+		UE_LOG(LogSockets, Log, TEXT("SteamSockets: Now tracking socket %u for addr %s, has parent? %d"), NewSocket->InternalHandle, *ForAddr.ToString(true), (ParentSocket != nullptr));
+		SocketInformationMap.Add(NewSocket->InternalHandle, NewSocketInfo);
+	}
+	else
+	{
+		UE_LOG(LogSockets, Warning, TEXT("SteamSockets: A socket with the address %s already exists, will not add another!"), *ForAddr.ToString(true));
+	}
+}
+
+void FSteamSocketsSubsystem::RemoveSocketsForListener(FSteamSocket* ListenerSocket)
+{
+	if (ListenerSocket == nullptr)
+	{
+		UE_LOG(LogSockets, Warning, TEXT("SteamSockets: Attempted to free sockets attached to an invalid listener!"));
+		return;
+	}
+
+	// When a listen socket closes in the SteamAPI, all sockets that were created from that listener are ungracefully closed.
+	// However this does not free any memory, nor does it alert any other object. As such, we have to process this destruction here.
+	UE_LOG(LogSockets, Log, TEXT("SteamSockets: Closing all sockets attached to listener %u"), ListenerSocket->InternalHandle);
+	for (SocketHandleInfoMap::TIterator It(SocketInformationMap); It; ++It)
+	{
+		FSteamSocketInformation& SocketInfo = It.Value();
+		if (SocketInfo.Parent == ListenerSocket)
+		{
+			UE_LOG(LogSockets, Verbose, TEXT("SteamSockets: Removed socket %u"), SocketInfo.Socket->InternalHandle);
+			// A loop later will clean out these connections, shutting them down properly.
+			// This is to prevent acting on deletion events from the API.
+			SocketInfo.Parent = nullptr;
+			// The children also should not be cleaning up any netdrivers they don't own
+			SocketInfo.NetDriver.Reset();
+			SocketInfo.MarkForDeletion();
+		}
+	}
+}
+
+void FSteamSocketsSubsystem::QueueRemoval(SteamSocketHandles RemoveHandle)
+{
+	FSteamSocketInformation* SocketInfo = GetSocketInfo(RemoveHandle);
+	if (SocketInfo != nullptr)
+	{
+		FString Address = (SocketInfo->Addr.IsValid()) ? SocketInfo->Addr->ToString(true) : TEXT("INVALID");
+		UE_LOG(LogSockets, Verbose, TEXT("SteamSockets: Marked socket %u with address %s for removal (pending)"), RemoveHandle, *Address);
+		SocketInfo->MarkForDeletion();
+	}
+}
+
+void FSteamSocketsSubsystem::LinkNetDriver(FSocket* Socket, USteamSocketsNetDriver* NewNetDriver)
+{
+	if (Socket == nullptr || NewNetDriver == nullptr)
+	{
+		UE_LOG(LogSockets, Warning, TEXT("SteamSockets: Attempted to create an invalid socket/netdriver pairing!"));
+		return;
+	}
+
+	FSteamSocket* SteamSocket = static_cast<FSteamSocket*>(Socket);
+	// Do not attempt to track netdrivers with invalid sockets.
+	if (SteamSocket->InternalHandle == k_HSteamListenSocket_Invalid || SteamSocket->InternalHandle == k_HSteamNetConnection_Invalid)
+	{
+		UE_LOG(LogSockets, Verbose, TEXT("SteamSockets: Dropped netdriver link with socket that has invalid handle."));
+		return;
+	}
+
+	// Link the netdriver to this socket.
+	FSteamSocketInformation* SocketInfo = GetSocketInfo(SteamSocket->InternalHandle);
+	if (SocketInfo != nullptr)
+	{
+		SocketInfo->NetDriver = NewNetDriver;
+	}
+}
+
+void FSteamSocketsSubsystem::AddDelayedListener(FSteamSocket* ListenSocket, USteamSocketsNetDriver* NewNetDriver)
+{
+	// Ignore any invalid data.
+	if (ListenSocket == nullptr || NewNetDriver == nullptr)
+	{
+		UE_LOG(LogSockets, Warning, TEXT("SteamSockets: Add delayed listener cannot listen to invalid data!"));
+		return;
+	}
+
+	FSteamPendingSocketInformation WatchedListener;
+	WatchedListener.Socket = ListenSocket;
+
+	// Track the netdriver
+	WatchedListener.NetDriver = NewNetDriver;
+
+	// Push to the array list
+	PendingListenerArray.Add(WatchedListener);
+}
+
+void FSteamSocketsSubsystem::OnServerLoginComplete(bool bWasSuccessful)
+{
+	TSharedPtr<FInternetAddrSteamSockets> SteamIdentityAddr = StaticCastSharedPtr<FInternetAddrSteamSockets>(GetIdentityAddress());
+	for (auto& DelayedListener : PendingListenerArray)
+	{
+		USteamSocketsNetDriver* SteamNetDriver = DelayedListener.NetDriver.Get();
+		if (SteamNetDriver == nullptr)
+		{
+			continue;
+		}
+
+		// If we are not successful, we need to get out.
+		if (!bWasSuccessful)
+		{
+			// Mark the NetDriver for cleanup.
+			SteamNetDriver->Shutdown();
+			continue;
+		}
+
+		// Update the binding address
+		if (SteamIdentityAddr.IsValid())
+		{
+			// Push in the new valid information.
+			DelayedListener.Socket->BindAddress.Addr = SteamIdentityAddr->Addr;
+			SteamNetDriver->LocalAddr = SteamIdentityAddr->Clone();
+			SteamNetDriver->LocalAddr->SetPlatformPort(DelayedListener.Socket->BindAddress.GetPlatformPort());
+
+			// Actually start the listen call
+			if (!DelayedListener.Socket->Listen(0))
+			{
+				// Connection failed, fire shutdown.
+				SteamNetDriver->Shutdown();
+				continue;
+			}
+		}
+
+		// Link the netdriver information so that we're ready.
+		LinkNetDriver(DelayedListener.Socket, SteamNetDriver);
+
+		// Toggle off the netdriver's tick throttle.
+		SteamNetDriver->bIsDelayedNetworkAccess = false;
+
+		UE_LOG(LogSockets, Verbose, TEXT("SteamSockets: Processed delayed listener: %s"), *DelayedListener.ToString());
+	}
+
+	PendingListenerArray.Empty();
+}
+
+TSharedPtr<FInternetAddr> FSteamSocketsSubsystem::GetIdentityAddress()
+{
+	SteamNetworkingIdentity SteamIDData;
+	// Attempt to get the machine's identity (their steam ID)
+	if (GetSteamSocketsInterface() && GetSteamSocketsInterface()->GetIdentity(&SteamIDData))
+	{
+		// If we got it, assign it to the SteamAddr
+		TSharedRef<FInternetAddrSteamSockets> SteamAddr = StaticCastSharedRef<FInternetAddrSteamSockets>(CreateInternetAddr());
+		SteamAddr->Addr = SteamIDData;
+		return SteamAddr;
+	}
+
+	return nullptr;
+}
+
+bool FSteamSocketsSubsystem::IsLoggedInToSteam() const
+{
+	// Check the game server functionality first because servers should be prioritized over SteamUser calls if
+	// the server has linked with the client dlls via launch arguments
+	return ((SteamGameServer() && IsRunningDedicatedServer()) ? SteamGameServer()->BLoggedOn() : ((SteamUser()) ? SteamUser()->BLoggedOn() : false));
+}
+
+// Helper function to translate Valve connection state messages into something human readable.
+static FString ConnectionStateToString(const ESteamNetworkingConnectionState& ConnectionState)
+{
+	switch (ConnectionState)
+	{
+	default:
+	case k_ESteamNetworkingConnectionState_None:
+		return TEXT("None");
+	case k_ESteamNetworkingConnectionState_Connecting:
+		return TEXT("Connecting");
+	case k_ESteamNetworkingConnectionState_FindingRoute:
+		return TEXT("Finding Route");
+	case k_ESteamNetworkingConnectionState_Connected:
+		return TEXT("Connected");
+	case k_ESteamNetworkingConnectionState_ClosedByPeer:
+		return TEXT("Closed by Peer");
+	case k_ESteamNetworkingConnectionState_ProblemDetectedLocally:
+		return TEXT("Local Connection Issue");
+	}
+
+	return TEXT("");
+}
+
+void FSteamSocketsSubsystem::SteamSocketEventHandler(struct SteamNetConnectionStatusChangedCallback_t* Message)
+{
+	if (Message == nullptr || FSteamSocketsSubsystem::GetSteamSocketsInterface() == nullptr)
+	{
+		return;
+	}
+
+	// New Connection came in
+	if (Message->m_eOldState == k_ESteamNetworkingConnectionState_None &&
+		Message->m_info.m_eState == k_ESteamNetworkingConnectionState_Connecting &&
+		Message->m_info.m_hListenSocket != k_HSteamListenSocket_Invalid)
+	{
+		// Push this to the netdriver.
+		FSteamSocketInformation* SocketInfo = GetSocketInfo(Message->m_info.m_hListenSocket);
+		USteamSocketsNetDriver* NetDriver = SocketInfo->NetDriver.Get();
+		if (SocketInfo && NetDriver)
+		{
+			NetDriver->OnConnectionCreated(Message->m_info.m_hListenSocket, Message->m_hConn);
+		}
+	}
+	else if ((Message->m_eOldState == k_ESteamNetworkingConnectionState_Connecting || Message->m_eOldState == k_ESteamNetworkingConnectionState_FindingRoute) &&
+		Message->m_info.m_eState == k_ESteamNetworkingConnectionState_Connected)
+	{
+		// Connection has been established
+		FSteamSocketInformation* SocketInfo = GetSocketInfo(Message->m_hConn);
+		USteamSocketsNetDriver* NetDriver = SocketInfo->NetDriver.Get();
+		if (SocketInfo && NetDriver)
+		{
+			NetDriver->OnConnectionUpdated(Message->m_hConn, (int32)Message->m_info.m_eState);
+		}
+	}
+	else if ((Message->m_eOldState == k_ESteamNetworkingConnectionState_FindingRoute || Message->m_eOldState == k_ESteamNetworkingConnectionState_Connecting 
+		|| Message->m_eOldState == k_ESteamNetworkingConnectionState_Connected) &&
+		(Message->m_info.m_eState == k_ESteamNetworkingConnectionState_ClosedByPeer || Message->m_info.m_eState == k_ESteamNetworkingConnectionState_ProblemDetectedLocally))
+	{
+		FSteamSocketInformation* SocketInfo = GetSocketInfo(Message->m_hConn);
+		// If we know anything about the socket, we need to act on it.
+		if (SocketInfo != nullptr && SocketInfo->Socket != nullptr && !SocketInfo->Socket->bIsListenSocket)
+		{
+			// Print out information about the disconnection.
+			UE_LOG(LogSockets, Verbose, TEXT("SteamSockets: Connection %u has disconnected. Old state: %s Reason: %s"), Message->m_hConn,
+				*ConnectionStateToString(Message->m_eOldState), *ConnectionStateToString(Message->m_info.m_eState));
+
+			// Let the netdriver know of the disconnection
+			USteamSocketsNetDriver* NetDriver = SocketInfo->NetDriver.Get();
+			if (NetDriver)
+			{
+				NetDriver->OnConnectionDisconnected(Message->m_hConn);
+			}
+
+			// If we close via the listen socket deletion then let the connection just die.
+			// Do not alert anyone as we already know we're disconnected.
+			//
+			// To determine this case we follow these two rules:
+			// * Servers will always mark all related sockets as not having a parent.
+			// * Clients will never have valid listener socket handles in their message data.
+			if (SocketInfo->Parent != nullptr || Message->m_info.m_hListenSocket == k_HSteamListenSocket_Invalid)
+			{
+				// If we are in here, it means we are a genuine disconnection, not caused by the listen socket destruction
+				// Or any other event that we ourselves caused.
+				//
+				// Mark this to be cleaned up. The tick will handle actual deletion.
+				SocketInfo->MarkForDeletion();
+			}
+		}
+	}
+}
diff --git a/Engine/Plugins/Runtime/Steam/SteamSockets/Source/SteamSockets/Private/SteamSocketsTaskManager.cpp b/Engine/Plugins/Runtime/Steam/SteamSockets/Source/SteamSockets/Private/SteamSocketsTaskManager.cpp
new file mode 100644
--- /dev/null
+++ b/Engine/Plugins/Runtime/Steam/SteamSockets/Source/SteamSockets/Private/SteamSocketsTaskManager.cpp
@@ -0,0 +1,41 @@
+// Copyright 1998-2019 Epic Games, Inc. All Rights Reserved.
+
+#include "SteamSocketsTaskManager.h"
+#include "SteamSocketsSubsystem.h"
+
+FSteamSocketsTaskManager::~FSteamSocketsTaskManager()
+{
+	MessageQueue.Empty();
+}
+
+void FSteamSocketsTaskManager::OnConnectionStatusChanged(SteamNetConnectionStatusChangedCallback_t* SteamConnectionMessage)
+{
+	// When we are called in here, we are likely going to be on the Steam OSS's online thread.
+	// The Steam callback API will fire on whatever thread the OSS is on.
+	//
+	// Because of how the callback api works, we need to store copies of these messages, as the pointers will be invalid
+	// the instant we leave our function.
+	MessageQueue.Enqueue(*SteamConnectionMessage);
+}
+
+void FSteamSocketsTaskManager::OnConnectionStatusChangedGS(SteamNetConnectionStatusChangedCallback_t* SteamConnectionMessage)
+{
+	MessageQueue.Enqueue(*SteamConnectionMessage);
+}
+
+void FSteamSocketsTaskManager::Tick()
+{
+	SteamNetConnectionStatusChangedCallback_t CallbackData;
+	while (!MessageQueue.IsEmpty())
+	{
+		FMemory::Memzero(CallbackData);
+		if (SocketSubsystem && MessageQueue.Dequeue(CallbackData))
+		{
+			SocketSubsystem->SteamSocketEventHandler(&CallbackData);
+		}
+		else
+		{
+			break;
+		}
+	}
+}
diff --git a/Engine/Plugins/Runtime/Steam/SteamSockets/Source/SteamSockets/Private/SteamSocketsTaskManager.h b/Engine/Plugins/Runtime/Steam/SteamSockets/Source/SteamSockets/Private/SteamSocketsTaskManager.h
new file mode 100644
--- /dev/null
+++ b/Engine/Plugins/Runtime/Steam/SteamSockets/Source/SteamSockets/Private/SteamSocketsTaskManager.h
@@ -0,0 +1,34 @@
+// Copyright 1998-2019 Epic Games, Inc. All Rights Reserved.
+
+#pragma once
+
+#include "CoreMinimal.h"
+#include "SteamSocketsPrivate.h"
+#include "Containers/Queue.h"
+
+// A glorified event manager for handling status updates about SteamSockets we hold.
+class FSteamSocketsTaskManager
+{
+public:
+	FSteamSocketsTaskManager(class FSteamSocketsSubsystem* SocketSub) :
+	    OnConnectionStatusChangedCallback(this, &FSteamSocketsTaskManager::OnConnectionStatusChanged),
+		OnConnectionStatusChangedGSCallback(this, &FSteamSocketsTaskManager::OnConnectionStatusChangedGS),
+		SocketSubsystem(SocketSub)
+	{
+		
+	}
+
+	~FSteamSocketsTaskManager();
+
+	void Tick();
+
+
+private:
+	STEAM_CALLBACK(FSteamSocketsTaskManager, OnConnectionStatusChanged, SteamNetConnectionStatusChangedCallback_t, OnConnectionStatusChangedCallback);
+	STEAM_GAMESERVER_CALLBACK(FSteamSocketsTaskManager, OnConnectionStatusChangedGS, SteamNetConnectionStatusChangedCallback_t, OnConnectionStatusChangedGSCallback);
+
+	class FSteamSocketsSubsystem* SocketSubsystem;
+	// A copied queue of our messages we get from the SteamAPI. There's no difference between client and server messages
+	// just how they are registered.
+	TQueue<SteamNetConnectionStatusChangedCallback_t> MessageQueue;
+};
diff --git a/Engine/Plugins/Runtime/Steam/SteamSockets/Source/SteamSockets/Public/SteamSocketsModule.h b/Engine/Plugins/Runtime/Steam/SteamSockets/Source/SteamSockets/Public/SteamSocketsModule.h
new file mode 100644
--- /dev/null
+++ b/Engine/Plugins/Runtime/Steam/SteamSockets/Source/SteamSockets/Public/SteamSocketsModule.h
@@ -0,0 +1,47 @@
+// Copyright 1998-2019 Epic Games, Inc. All Rights Reserved.
+
+#pragma once
+
+#include "CoreMinimal.h"
+#include "Modules/ModuleManager.h"
+#include "Modules/ModuleInterface.h"
+
+class STEAMSOCKETS_API FSteamSocketsModule : public IModuleInterface
+{
+public:
+
+	FSteamSocketsModule() : bEnabled(false)
+	{
+
+	}
+
+	virtual ~FSteamSocketsModule()
+	{
+	}
+
+	// IModuleInterface
+
+	virtual void StartupModule() override;
+	virtual void ShutdownModule() override;
+
+	virtual bool SupportsDynamicReloading() override
+	{
+		return false;
+	}
+
+	static inline class FSteamSocketsModule& Get()
+	{
+		return FModuleManager::LoadModuleChecked<class FSteamSocketsModule>("SteamSockets");
+	}
+
+	static inline bool IsAvailable()
+	{
+		return FModuleManager::Get().IsModuleLoaded("SteamSockets");
+	}
+	
+	bool IsSteamSocketsEnabled() const { return bEnabled; }
+
+private:
+	// Allows for easier configuration of which SteamSockets system to use.
+	bool bEnabled;
+};
diff --git a/Engine/Plugins/Runtime/Steam/SteamSockets/Source/SteamSockets/Public/SteamSocketsNetConnection.h b/Engine/Plugins/Runtime/Steam/SteamSockets/Source/SteamSockets/Public/SteamSocketsNetConnection.h
new file mode 100644
--- /dev/null
+++ b/Engine/Plugins/Runtime/Steam/SteamSockets/Source/SteamSockets/Public/SteamSocketsNetConnection.h
@@ -0,0 +1,51 @@
+// Copyright 1998-2019 Epic Games, Inc. All Rights Reserved.
+
+#pragma once
+
+#include "Engine/NetConnection.h"
+#include "SteamSocketsNetConnection.generated.h"
+
+// Forward declare some types for functional type arguments
+class FSocket;
+class UNetDriver;
+class FInternetAddr;
+class FSteamSocket;
+
+UCLASS(transient, config=Engine)
+class STEAMSOCKETS_API USteamSocketsNetConnection : public UNetConnection
+{
+	GENERATED_BODY()
+
+public:
+
+	USteamSocketsNetConnection() :
+		ConnectionSocket(nullptr),
+		bInConnectionlessHandshake(false)
+	{
+	}
+
+	//~ Begin NetConnection Interface
+	virtual void CleanUp() override;
+	virtual void InitBase(UNetDriver* InDriver, FSocket* InSocket, const FURL& InURL, EConnectionState InState, int32 InMaxPacket = 0, int32 InPacketOverhead = 0) override;
+	virtual void InitRemoteConnection(UNetDriver* InDriver, FSocket* InSocket, const FURL& InURL, const FInternetAddr& InRemoteAddr, 
+		EConnectionState InState, int32 InMaxPacket = 0, int32 InPacketOverhead = 0) override;
+	virtual void InitLocalConnection(UNetDriver* InDriver, FSocket* InSocket, const FURL& InURL, EConnectionState InState, int32 InMaxPacket = 0, int32 InPacketOverhead = 0) override;
+	virtual void LowLevelSend(void* Data, int32 CountBits, FOutPacketTraits& Traits) override;
+	FString LowLevelGetRemoteAddress(bool bAppendPort=false) override;
+	FString LowLevelDescribe() override;
+	//~ End NetConnection Interface
+
+private:
+	const FSteamSocket* GetRawSocket() const { return ConnectionSocket; }
+
+	void HandleRecvMessage(void* InData, int32 SizeOfData, const FInternetAddr* InFormattedAddress);
+
+	void FlagForHandshake() { bInConnectionlessHandshake = true; }
+
+	void ClearSocket() { ConnectionSocket = nullptr; }
+
+	FSteamSocket* ConnectionSocket;
+	bool bInConnectionlessHandshake;
+
+	friend class USteamSocketsNetDriver;
+};
diff --git a/Engine/Plugins/Runtime/Steam/SteamSockets/Source/SteamSockets/Public/SteamSocketsNetDriver.h b/Engine/Plugins/Runtime/Steam/SteamSockets/Source/SteamSockets/Public/SteamSocketsNetDriver.h
new file mode 100644
--- /dev/null
+++ b/Engine/Plugins/Runtime/Steam/SteamSockets/Source/SteamSockets/Public/SteamSocketsNetDriver.h
@@ -0,0 +1,59 @@
+// Copyright 1998-2019 Epic Games, Inc. All Rights Reserved.
+
+#pragma once
+
+#include "CoreMinimal.h"
+#include "UObject/ObjectMacros.h"
+#include "Engine/NetDriver.h"
+#include "SteamSocketsPackage.h"
+#include "SteamSocketsTypes.h"
+#include "SteamSocketsNetDriver.generated.h"
+
+class FNetworkNotify;
+
+UCLASS(transient, config=Engine)
+class STEAMSOCKETS_API USteamSocketsNetDriver : public UNetDriver
+{
+	GENERATED_BODY()
+
+public:
+
+	USteamSocketsNetDriver() :
+		Socket(nullptr),
+		bIsDelayedNetworkAccess(false)
+	{
+	}
+
+	//~ Begin UObject Interface
+	virtual void PostInitProperties() override;
+	//~ End UObject Interface
+
+	//~ Begin UNetDriver Interface.
+	virtual void Shutdown() override;
+	virtual bool IsAvailable() const override;
+	virtual bool InitBase(bool bInitAsClient, FNetworkNotify* InNotify, const FURL& URL, bool bReuseAddressAndPort, FString& Error) override;
+	virtual bool InitConnect(FNetworkNotify* InNotify, const FURL& ConnectURL, FString& Error) override;
+	virtual bool InitListen(FNetworkNotify* InNotify, FURL& LocalURL, bool bReuseAddressAndPort, FString& Error) override;
+	virtual void TickDispatch(float DeltaTime) override;
+	virtual void LowLevelSend(TSharedPtr<const FInternetAddr> Address, void* Data, int32 CountBits, FOutPacketTraits& Traits) override;
+	virtual void LowLevelDestroy() override;
+	virtual class ISocketSubsystem* GetSocketSubsystem() override;
+	virtual bool IsNetResourceValid(void) override;
+	//~ End UNetDriver Interface
+
+	bool ArePacketHandlersDisabled() const;
+
+protected:
+	class FSteamSocket* Socket;
+	bool bIsDelayedNetworkAccess;
+
+	void ResetSocketInfo(const class FSteamSocket* RemovedSocket);
+
+	UNetConnection* FindClientConnectionForHandle(SteamSocketHandles SocketHandle);
+
+	void OnConnectionCreated(SteamSocketHandles ListenParentHandle, SteamSocketHandles SocketHandle);
+	void OnConnectionUpdated(SteamSocketHandles SocketHandle, int32 NewState);
+	void OnConnectionDisconnected(SteamSocketHandles SocketHandle);
+
+	friend class FSteamSocketsSubsystem;
+};
diff --git a/Engine/Plugins/Runtime/Steam/SteamSockets/Source/SteamSockets/Public/SteamSocketsPackage.h b/Engine/Plugins/Runtime/Steam/SteamSockets/Source/SteamSockets/Public/SteamSocketsPackage.h
new file mode 100644
--- /dev/null
+++ b/Engine/Plugins/Runtime/Steam/SteamSockets/Source/SteamSockets/Public/SteamSocketsPackage.h
@@ -0,0 +1,15 @@
+// Copyright 1998-2019 Epic Games, Inc. All Rights Reserved.
+
+// Can't be #pragma once because other modules may define PACKAGE_SCOPE
+
+// Intended to be the last include in an exported class definition
+// Properly defines some members as "public to the module" vs "private to the consumer/user"
+
+// [[ IncludeTool: Inline ]]
+
+#undef PACKAGE_SCOPE
+#ifdef STEAMSOCKETS_MODULE
+#define PACKAGE_SCOPE public
+#else
+#define PACKAGE_SCOPE protected
+#endif
diff --git a/Engine/Plugins/Runtime/Steam/SteamSockets/Source/SteamSockets/Public/SteamSocketsSubsystem.h b/Engine/Plugins/Runtime/Steam/SteamSockets/Source/SteamSockets/Public/SteamSocketsSubsystem.h
new file mode 100644
--- /dev/null
+++ b/Engine/Plugins/Runtime/Steam/SteamSockets/Source/SteamSockets/Public/SteamSocketsSubsystem.h
@@ -0,0 +1,200 @@
+// Copyright 1998-2019 Epic Games, Inc. All Rights Reserved.
+
+#pragma once
+
+#include "CoreMinimal.h"
+#include "SocketSubsystem.h"
+#include "IPAddress.h"
+#include "Containers/Ticker.h"
+#include "Containers/Map.h"
+#include "Containers/Array.h"
+#include "UObject/WeakObjectPtr.h"
+#include "SteamSocketsPackage.h"
+#include "SteamSocketsTypes.h"
+#include "SocketTypes.h"
+
+// Forward declare some internal types for bookkeeping.
+class FSteamSocket;
+class USteamSocketsNetDriver;
+
+/**
+ * Steam Sockets specific socket subsystem implementation. 
+ * This class can only be used with the SteamSocketsNetDriver and the SteamSocketsNetConnection classes.
+ * This subsystem does not support mixing any other NetDriver/NetConnection format. Doing so will cause this protocol to not function.
+ */
+class STEAMSOCKETS_API FSteamSocketsSubsystem : public ISocketSubsystem, public FTickerObjectBase, public FSelfRegisteringExec
+{
+public:
+
+	FSteamSocketsSubsystem() :
+		LastSocketError(0),
+		bShouldTestPeek(false),
+		SteamEventManager(nullptr),
+		bUseRelays(true),
+		SteamAPIClientHandle(nullptr),
+		SteamAPIServerHandle(nullptr)
+	{
+	}
+
+	//~ Begin SocketSubsystem Interface
+	virtual bool Init(FString& Error) override;
+	virtual void Shutdown() override;
+
+	virtual class FSocket* CreateSocket(const FName& SocketType, const FString& SocketDescription, const FName& ProtocolType) override;
+	virtual void DestroySocket(class FSocket* Socket) override;
+
+	virtual FAddressInfoResult GetAddressInfo(const TCHAR* HostName, const TCHAR* ServiceName = nullptr,
+		EAddressInfoFlags QueryFlags = EAddressInfoFlags::Default,
+		const FName ProtocolTypeName = NAME_None,
+		ESocketType SocketType = ESocketType::SOCKTYPE_Unknown) override;
+	virtual TSharedPtr<FInternetAddr> GetAddressFromString(const FString& IPAddress) override;
+
+	virtual bool GetHostName(FString& HostName) override;
+
+	virtual TSharedRef<FInternetAddr> CreateInternetAddr() override;
+
+	virtual const TCHAR* GetSocketAPIName() const override;
+
+	virtual ESocketErrors GetLastErrorCode() override {	return (ESocketErrors)LastSocketError; }
+	virtual ESocketErrors TranslateErrorCode(int32 Code) override {	return (ESocketErrors)Code;	}
+
+	virtual bool GetLocalAdapterAddresses(FSharedInternetAddrArray& OutAddresses) override;
+	virtual TSharedRef<FInternetAddr> GetLocalBindAddr(FOutputDevice& Out) override;
+	virtual bool HasNetworkDevice() override { return true; }
+	virtual bool IsSocketWaitSupported() const override { return false; }
+	virtual bool RequiresChatDataBeSeparate() override { return false; }
+	virtual bool RequiresEncryptedPackets() override { return false; }
+	//~ End SocketSubsystem Interface
+
+	//~ Begin FTickerObject Interface
+	virtual bool Tick(float DeltaTime) override;
+	//~ End FTickerObject Interface
+
+	//~ Begin FSelfRegisteringExec Interface
+	virtual bool Exec(class UWorld* InWorld, const TCHAR* Cmd, FOutputDevice& Ar) override;
+	//~ End FSelfRegisteringExec Interface
+
+	bool IsUsingRelayNetwork() const { return bUseRelays; }
+
+	bool IsSteamInitialized() const { return SteamAPIClientHandle.IsValid() || SteamAPIServerHandle.IsValid(); }
+
+	static class ISteamNetworkingSockets* GetSteamSocketsInterface();
+
+PACKAGE_SCOPE:
+	/** A struct for holding steam socket information and managing bookkeeping on the protocol. */
+	struct FSteamSocketInformation
+	{
+		FSteamSocketInformation(TSharedPtr<FInternetAddr> InAddr, FSteamSocket* InSocket, FSteamSocket* InParent = nullptr) :
+			Addr(InAddr),
+			Socket(InSocket),
+			Parent(InParent),
+			NetDriver(nullptr),
+			bMarkedForDeletion(false)
+		{
+		}
+
+		void MarkForDeletion();
+
+		bool IsMarkedForDeletion() const { return bMarkedForDeletion; }
+
+		bool operator==(const FSteamSocket* RHS) const
+		{
+			return Socket == RHS;
+		}
+
+		bool operator==(const FInternetAddr& InAddr) const;
+
+		bool operator==(const FSteamSocketInformation& RHS) const
+		{
+			return RHS.Addr == Addr && RHS.Socket == Socket;
+		}
+
+		bool IsValid() const
+		{
+			return Addr.IsValid() && Parent != nullptr;
+		}
+
+		FString ToString() const;
+
+		TSharedPtr<FInternetAddr> Addr;
+		FSteamSocket* Socket;
+		// Sockets created from a listener have a parent
+		FSteamSocket* Parent;
+		// The NetDriver for this connection.
+		TWeakObjectPtr<USteamSocketsNetDriver> NetDriver;
+	private:
+		bool bMarkedForDeletion;
+	};
+
+	// Steam socket queriers
+	FSteamSocketInformation* GetSocketInfo(SteamSocketHandles InternalSocketHandle);
+	FSteamSocketInformation* GetSocketInfo(const FInternetAddr& ForAddress);
+
+	// Steam socket bookkeeping modifiers
+	void AddSocket(const FInternetAddr& ForAddr, FSteamSocket* NewSocket, FSteamSocket* ParentSocket = nullptr);
+	void RemoveSocketsForListener(FSteamSocket* ListenerSocket);
+	void QueueRemoval(SteamSocketHandles SocketHandle);
+	void LinkNetDriver(FSocket* Socket, USteamSocketsNetDriver* NewNetDriver);
+
+	// Delayed listen socket helpers.
+	void AddDelayedListener(FSteamSocket* ListenSocket, USteamSocketsNetDriver* NewNetDriver);
+	void OnServerLoginComplete(bool bWasSuccessful);
+
+	// Returns this machine's identity in the form of a FInternetAddrSteamSockets
+	TSharedPtr<FInternetAddr> GetIdentityAddress();
+	
+	// Returns if our account is currently logged into the Steam network
+	bool IsLoggedInToSteam() const;
+
+	/** Last error set by the socket subsystem or one of its sockets */
+	int32 LastSocketError;
+
+	// Singleton helpers
+	static FSteamSocketsSubsystem* Create();
+	static void Destroy();
+
+	// SteamAPI internals handler
+	void SteamSocketEventHandler(struct SteamNetConnectionStatusChangedCallback_t* ConnectionEvent);
+
+	/** Flag for testing peek messaging (only usable in non-shipping builds) */
+	bool bShouldTestPeek;
+
+protected:
+	void CleanSocketInformation(bool bForceClean);
+	void DumpSocketInformationMap() const;
+
+	/** Single instantiation of this subsystem */
+	static FSteamSocketsSubsystem* SocketSingleton;
+
+	/** Event manager for Steam tasks */
+	TUniquePtr<class FSteamSocketsTaskManager> SteamEventManager;
+
+	/** Determines if the connections are going to be using the relay network */
+	bool bUseRelays;
+
+	/** Steam Client API Handle */
+	TSharedPtr<class FSteamClientInstanceHandler> SteamAPIClientHandle;
+
+	/** Steam Server API Handle */
+	TSharedPtr<class FSteamServerInstanceHandler> SteamAPIServerHandle;
+
+	/** Active connection bookkeeping */
+	typedef TMap<SteamSocketHandles, FSteamSocketInformation> SocketHandleInfoMap;
+	SocketHandleInfoMap SocketInformationMap;
+
+	/** Structure for handling sockets that cannot be established due to platform login (for listener sockets) */
+	struct FSteamPendingSocketInformation
+	{
+		FSteamSocket* Socket;
+		TWeakObjectPtr<USteamSocketsNetDriver> NetDriver;
+
+		FString ToString() const;
+	};
+
+	// Array of listeners we need to activate.
+	TArray<FSteamPendingSocketInformation> PendingListenerArray;
+
+	// Delegate handle for handling when a dedicated server logs into the Steam platform
+	FDelegateHandle SteamServerLoginDelegateHandle;
+};
+
diff --git a/Engine/Plugins/Runtime/Steam/SteamSockets/Source/SteamSockets/Public/SteamSocketsTypes.h b/Engine/Plugins/Runtime/Steam/SteamSockets/Source/SteamSockets/Public/SteamSocketsTypes.h
new file mode 100644
--- /dev/null
+++ b/Engine/Plugins/Runtime/Steam/SteamSockets/Source/SteamSockets/Public/SteamSocketsTypes.h
@@ -0,0 +1,18 @@
+// Copyright 1998-2019 Epic Games, Inc. All Rights Reserved.
+
+#pragma once
+
+#include "Engine/EngineBaseTypes.h"
+
+#ifndef STEAM_SOCKETS_SUBSYSTEM
+#define STEAM_SOCKETS_SUBSYSTEM FName(TEXT("SteamSockets"))
+#endif
+
+// Add our two new protocol types
+namespace FNetworkProtocolTypes
+{
+	STEAMSOCKETS_API extern const FName SteamSocketsP2P;
+	STEAMSOCKETS_API extern const FName SteamSocketsIP;
+}
+
+typedef uint32 SteamSocketHandles;
diff --git a/Engine/Plugins/Runtime/Steam/SteamSockets/Source/SteamSockets/SteamSockets.Build.cs b/Engine/Plugins/Runtime/Steam/SteamSockets/Source/SteamSockets/SteamSockets.Build.cs
new file mode 100644
--- /dev/null
+++ b/Engine/Plugins/Runtime/Steam/SteamSockets/Source/SteamSockets/SteamSockets.Build.cs
@@ -0,0 +1,27 @@
+// Copyright 1998-2019 Epic Games, Inc. All Rights Reserved.
+
+using UnrealBuildTool;
+using System.IO;
+
+public class SteamSockets : ModuleRules
+{
+    public SteamSockets(ReadOnlyTargetRules Target) : base(Target)
+    {
+		PublicDefinitions.Add("STEAMSOCKETS_MODULE=1");
+
+		PrivateDependencyModuleNames.AddRange(
+			new string[] {
+				"Core", 
+				"CoreUObject",
+				"Engine", 
+				"Sockets",
+				"OnlineSubsystem",
+				"OnlineSubsystemSteam",
+				"PacketHandler",
+                "SteamShared"
+            }
+		);
+
+		AddEngineThirdPartyPrivateStaticDependencies(Target, "Steamworks");
+	}
+}
diff --git a/Engine/Plugins/Runtime/Steam/SteamSockets/SteamSockets.uplugin b/Engine/Plugins/Runtime/Steam/SteamSockets/SteamSockets.uplugin
new file mode 100644
--- /dev/null
+++ b/Engine/Plugins/Runtime/Steam/SteamSockets/SteamSockets.uplugin
@@ -0,0 +1,43 @@
+{
+	"FileVersion" : 3,
+	"Version" : 1.0,
+	"VersionName" : "1.0",
+	"FriendlyName" : "Steam Sockets",
+	"Description" : "New Steamworks Networking code that supports the SteamSockets interface.",
+	"Category" : "Networking",
+	"CreatedBy" : "Epic Games, Inc.",
+	"CreatedByURL" : "https://epicgames.com",
+	"EnabledByDefault" : false,
+	"CanContainContent" : false,
+	"IsBetaVersion" : false,
+	"Installed" : false,
+	"Modules" :
+	[
+		{
+			"Name" : "SteamSockets",
+			"Type" : "Runtime",
+			"LoadingPhase" : "Default",
+			"WhitelistPlatforms" :
+			[
+				"Win32",
+				"Win64",
+				"Mac",
+				"Linux"
+			]
+		}
+	],
+	"Plugins": [
+		{
+			"Name": "OnlineSubsystem",
+			"Enabled": true
+		},
+		{
+			"Name": "OnlineSubsystemSteam",
+			"Enabled": true
+		},
+		{
+			"Name": "SteamShared",
+			"Enabled": true
+		}
+	]
+}
\ No newline at end of file
