Fix CIS errors due to some networking layer improvements not having been committed yet.

Also fix the TUniquePtr template destructor issue which only showed up on a full rebuild.

#rnx

diff --git a/Engine/Plugins/Runtime/Steam/SteamSockets/Source/SteamSockets/Private/SteamSocketsNetDriver.cpp b/Engine/Plugins/Runtime/Steam/SteamSockets/Source/SteamSockets/Private/SteamSocketsNetDriver.cpp
--- a/Engine/Plugins/Runtime/Steam/SteamSockets/Source/SteamSockets/Private/SteamSocketsNetDriver.cpp
+++ b/Engine/Plugins/Runtime/Steam/SteamSockets/Source/SteamSockets/Private/SteamSocketsNetDriver.cpp
@@ -126,7 +126,7 @@ bool USteamSocketsNetDriver::InitBase(bool bInitAsClient, FNetworkNotify* InNoti
 	if (!bIsDelayedNetworkAccess)
 	{
 		// Set the binding address.
-		FSharedInternetAddrArray Addresses;
+		TArray<TSharedPtr<FInternetAddr>> Addresses;
 		if (SocketSub->GetLocalAdapterAddresses(Addresses))
 		{
 			// Attempt to find a binding address that matches protocol of the connection destination
diff --git a/Engine/Plugins/Runtime/Steam/SteamSockets/Source/SteamSockets/Private/SteamSocketsSubsystem.cpp b/Engine/Plugins/Runtime/Steam/SteamSockets/Source/SteamSockets/Private/SteamSocketsSubsystem.cpp
--- a/Engine/Plugins/Runtime/Steam/SteamSockets/Source/SteamSockets/Private/SteamSocketsSubsystem.cpp
+++ b/Engine/Plugins/Runtime/Steam/SteamSockets/Source/SteamSockets/Private/SteamSocketsSubsystem.cpp
@@ -288,7 +288,7 @@ FAddressInfoResult FSteamSocketsSubsystem::GetAddressInfo(const TCHAR* HostName,
 		return ResultData;
 	}
 
-	FSharedInternetAddrArray AddressList;
+	TArray<TSharedPtr<FInternetAddr>> AddressList;
 	// Handle binding flags or if you don't pass in a hostname
 	if (HostName == nullptr || EnumHasAnyFlags(QueryFlags, EAddressInfoFlags::BindableAddress))
 	{
@@ -398,7 +398,7 @@ const TCHAR* FSteamSocketsSubsystem::GetSocketAPIName() const
 	return TEXT("SteamSockets");
 }
 
-bool FSteamSocketsSubsystem::GetLocalAdapterAddresses(FSharedInternetAddrArray& OutAddresses)
+bool FSteamSocketsSubsystem::GetLocalAdapterAddresses(TArray<TSharedPtr<FInternetAddr>>& OutAddresses)
 {
 	// Add Multihome
 	TSharedRef<FInternetAddr> MultihomeAddress = CreateInternetAddr();
@@ -427,7 +427,7 @@ bool FSteamSocketsSubsystem::GetLocalAdapterAddresses(FSharedInternetAddrArray&
 
 TSharedRef<FInternetAddr> FSteamSocketsSubsystem::GetLocalBindAddr(FOutputDevice& Out)
 {
-	FSharedInternetAddrArray Adapters;
+	TArray<TSharedPtr<FInternetAddr>> Adapters;
 	if (GetLocalAdapterAddresses(Adapters))
 	{
 		return Adapters[0].ToSharedRef();
diff --git a/Engine/Plugins/Runtime/Steam/SteamSockets/Source/SteamSockets/Private/SteamSocketsTaskManager.cpp b/Engine/Plugins/Runtime/Steam/SteamSockets/Source/SteamSockets/Private/SteamSocketsTaskManager.cpp
--- a/Engine/Plugins/Runtime/Steam/SteamSockets/Source/SteamSockets/Private/SteamSocketsTaskManager.cpp
+++ b/Engine/Plugins/Runtime/Steam/SteamSockets/Source/SteamSockets/Private/SteamSocketsTaskManager.cpp
@@ -3,11 +3,6 @@
 #include "SteamSocketsTaskManager.h"
 #include "SteamSocketsSubsystem.h"
 
-FSteamSocketsTaskManager::~FSteamSocketsTaskManager()
-{
-	MessageQueue.Empty();
-}
-
 void FSteamSocketsTaskManager::OnConnectionStatusChanged(SteamNetConnectionStatusChangedCallback_t* SteamConnectionMessage)
 {
 	// When we are called in here, we are likely going to be on the Steam OSS's online thread.
diff --git a/Engine/Plugins/Runtime/Steam/SteamSockets/Source/SteamSockets/Private/SteamSocketsTaskManager.h b/Engine/Plugins/Runtime/Steam/SteamSockets/Source/SteamSockets/Private/SteamSocketsTaskManager.h
--- a/Engine/Plugins/Runtime/Steam/SteamSockets/Source/SteamSockets/Private/SteamSocketsTaskManager.h
+++ b/Engine/Plugins/Runtime/Steam/SteamSockets/Source/SteamSockets/Private/SteamSocketsTaskManager.h
@@ -3,31 +3,34 @@
 #pragma once
 
 #include "CoreMinimal.h"
+#include "SteamSocketsTaskManagerInterface.h"
 #include "SteamSocketsPrivate.h"
 #include "Containers/Queue.h"
 
 // A glorified event manager for handling status updates about SteamSockets we hold.
-class FSteamSocketsTaskManager
+class FSteamSocketsTaskManager : public FSteamSocketsTaskManagerInterface
 {
 public:
-	FSteamSocketsTaskManager(class FSteamSocketsSubsystem* SocketSub) :
+	FSteamSocketsTaskManager(FSteamSocketsSubsystem* SocketSub) :
+		FSteamSocketsTaskManagerInterface(SocketSub),
 	    OnConnectionStatusChangedCallback(this, &FSteamSocketsTaskManager::OnConnectionStatusChanged),
-		OnConnectionStatusChangedGSCallback(this, &FSteamSocketsTaskManager::OnConnectionStatusChangedGS),
-		SocketSubsystem(SocketSub)
+		OnConnectionStatusChangedGSCallback(this, &FSteamSocketsTaskManager::OnConnectionStatusChangedGS)
 	{
 		
 	}
 
-	~FSteamSocketsTaskManager();
+	virtual ~FSteamSocketsTaskManager()
+	{
+		MessageQueue.Empty();
+	}
 
-	void Tick();
+	virtual void Tick() override;
 
 
 private:
 	STEAM_CALLBACK(FSteamSocketsTaskManager, OnConnectionStatusChanged, SteamNetConnectionStatusChangedCallback_t, OnConnectionStatusChangedCallback);
 	STEAM_GAMESERVER_CALLBACK(FSteamSocketsTaskManager, OnConnectionStatusChangedGS, SteamNetConnectionStatusChangedCallback_t, OnConnectionStatusChangedGSCallback);
 
-	class FSteamSocketsSubsystem* SocketSubsystem;
 	// A copied queue of our messages we get from the SteamAPI. There's no difference between client and server messages
 	// just how they are registered.
 	TQueue<SteamNetConnectionStatusChangedCallback_t> MessageQueue;
diff --git a/Engine/Plugins/Runtime/Steam/SteamSockets/Source/SteamSockets/Public/SteamSocketsSubsystem.h b/Engine/Plugins/Runtime/Steam/SteamSockets/Source/SteamSockets/Public/SteamSocketsSubsystem.h
--- a/Engine/Plugins/Runtime/Steam/SteamSockets/Source/SteamSockets/Public/SteamSocketsSubsystem.h
+++ b/Engine/Plugins/Runtime/Steam/SteamSockets/Source/SteamSockets/Public/SteamSocketsSubsystem.h
@@ -10,6 +10,7 @@
 #include "Containers/Array.h"
 #include "UObject/WeakObjectPtr.h"
 #include "SteamSocketsPackage.h"
+#include "SteamSocketsTaskManagerInterface.h"
 #include "SteamSocketsTypes.h"
 #include "SocketTypes.h"
 
@@ -58,7 +59,7 @@ public:
 	virtual ESocketErrors GetLastErrorCode() override {	return (ESocketErrors)LastSocketError; }
 	virtual ESocketErrors TranslateErrorCode(int32 Code) override {	return (ESocketErrors)Code;	}
 
-	virtual bool GetLocalAdapterAddresses(FSharedInternetAddrArray& OutAddresses) override;
+	virtual bool GetLocalAdapterAddresses(TArray<TSharedPtr<FInternetAddr>>& OutAddresses) override;
 	virtual TSharedRef<FInternetAddr> GetLocalBindAddr(FOutputDevice& Out) override;
 	virtual bool HasNetworkDevice() override { return true; }
 	virtual bool IsSocketWaitSupported() const override { return false; }
@@ -74,8 +75,10 @@ public:
 	virtual bool Exec(class UWorld* InWorld, const TCHAR* Cmd, FOutputDevice& Ar) override;
 	//~ End FSelfRegisteringExec Interface
 
+	/** Returns if the application is using the SteamSocket relays */
 	bool IsUsingRelayNetwork() const { return bUseRelays; }
 
+	/** Basic function to determine if Steam has been initialized properly. */
 	bool IsSteamInitialized() const { return SteamAPIClientHandle.IsValid() || SteamAPIServerHandle.IsValid(); }
 
 	static class ISteamNetworkingSockets* GetSteamSocketsInterface();
@@ -167,7 +170,7 @@ protected:
 	static FSteamSocketsSubsystem* SocketSingleton;
 
 	/** Event manager for Steam tasks */
-	TUniquePtr<class FSteamSocketsTaskManager> SteamEventManager;
+	TUniquePtr<class FSteamSocketsTaskManagerInterface> SteamEventManager;
 
 	/** Determines if the connections are going to be using the relay network */
 	bool bUseRelays;
@@ -197,4 +200,3 @@ protected:
 	// Delegate handle for handling when a dedicated server logs into the Steam platform
 	FDelegateHandle SteamServerLoginDelegateHandle;
 };
-
diff --git a/Engine/Plugins/Runtime/Steam/SteamSockets/Source/SteamSockets/Public/SteamSocketsTaskManagerInterface.h b/Engine/Plugins/Runtime/Steam/SteamSockets/Source/SteamSockets/Public/SteamSocketsTaskManagerInterface.h
new file mode 100644
--- /dev/null
+++ b/Engine/Plugins/Runtime/Steam/SteamSockets/Source/SteamSockets/Public/SteamSocketsTaskManagerInterface.h
@@ -0,0 +1,23 @@
+// Copyright 1998-2019 Epic Games, Inc. All Rights Reserved.
+
+#pragma once
+
+/** 
+ *  A base task manager interface for handling messages from the Steamworks SDK to pipe to the SteamSocketsSubsystem 
+ *  This is mostly an empty interface used to allow for creation of uniqueptrs of the task manager interface
+ *  without exposing Steam data types.
+ */
+class FSteamSocketsTaskManagerInterface
+{
+public:
+	FSteamSocketsTaskManagerInterface(class FSteamSocketsSubsystem* SocketSub) :
+		SocketSubsystem(SocketSub)
+	{
+	}
+
+	/** A function that will be called on every update tick of the SteamSocketsSubsystem */
+	virtual void Tick() = 0;
+
+protected:
+	class FSteamSocketsSubsystem* SocketSubsystem;
+};
