Fix Clang CIS error when building Linux with SteamSockets.

#rnx

diff --git a/Engine/Plugins/Runtime/Steam/SteamSockets/Source/SteamSockets/Public/SteamSocketsTaskManagerInterface.h b/Engine/Plugins/Runtime/Steam/SteamSockets/Source/SteamSockets/Public/SteamSocketsTaskManagerInterface.h
--- a/Engine/Plugins/Runtime/Steam/SteamSockets/Source/SteamSockets/Public/SteamSocketsTaskManagerInterface.h
+++ b/Engine/Plugins/Runtime/Steam/SteamSockets/Source/SteamSockets/Public/SteamSocketsTaskManagerInterface.h
@@ -14,6 +14,8 @@ public:
 		SocketSubsystem(SocketSub)
 	{
 	}
+	
+	virtual ~FSteamSocketsTaskManagerInterface() = default;
 
 	/** A function that will be called on every update tick of the SteamSocketsSubsystem */
 	virtual void Tick() = 0;
