Fixing static analysis warnings in SteamSockets

#rnx

diff --git a/Engine/Plugins/Runtime/Steam/SteamSockets/Source/SteamSockets/Private/SteamSocket.cpp b/Engine/Plugins/Runtime/Steam/SteamSockets/Source/SteamSockets/Private/SteamSocket.cpp
--- a/Engine/Plugins/Runtime/Steam/SteamSockets/Source/SteamSockets/Private/SteamSocket.cpp
+++ b/Engine/Plugins/Runtime/Steam/SteamSockets/Source/SteamSockets/Private/SteamSocket.cpp
@@ -37,12 +37,14 @@ bool FSteamSocket::Close()
 {
 	bool bWasSuccessful = false;
 
+	STEAM_SDK_IGNORE_REDUNDANCY_START
 	// If we're already closed, don't bother doing anything else.
 	if (InternalHandle == k_HSteamListenSocket_Invalid || InternalHandle == k_HSteamNetConnection_Invalid)
 	{
 		UE_LOG(LogSockets, VeryVerbose, TEXT("SteamSockets: Socket is already cleaned up, ready for destruction."));
 		return true;
 	}
+	STEAM_SDK_IGNORE_REDUNDANCY_END
 
 	ISteamNetworkingSockets* SocketInterface = FSteamSocketsSubsystem::GetSteamSocketsInterface();
 
@@ -348,8 +350,16 @@ bool FSteamSocket::HasPendingData(uint32& PendingDataSize)
 {
 	if (bHasPendingData)
 	{
-		PendingDataSize = PendingData->GetSize();
-		return true;
+		if (PendingData != nullptr)
+		{
+			PendingDataSize = PendingData->GetSize();
+			return true;
+		}
+		else
+		{
+			UE_LOG(LogSockets, Warning, TEXT("SteamSockets: HasPendingData flag is raised but the pendingdata object is null!"));
+			bHasPendingData = false;
+		}
 	}
 
 	int32 MessagesRead;
@@ -434,6 +444,7 @@ bool FSteamSocket::GetPeerAddress(FInternetAddr& OutAddr)
 bool FSteamSocket::SetNoDelay(bool bIsNoDelay)
 {
 	// See https://partner.steamgames.com/doc/api/steamnetworkingtypes#message_sending_flags for what all these flags mean
+	STEAM_SDK_IGNORE_REDUNDANCY_START
 	if (bIsNoDelay)
 	{
 		if (SendMode == k_nSteamNetworkingSend_Unreliable || SendMode == k_nSteamNetworkingSend_NoNagle || SendMode == k_nSteamNetworkingSend_UnreliableNoNagle)
@@ -460,7 +471,7 @@ bool FSteamSocket::SetNoDelay(bool bIsNoDelay)
 			SendMode = k_nSteamNetworkingSend_Reliable;
 		}
 	}
-
+	STEAM_SDK_IGNORE_REDUNDANCY_END
 	return true;
 }
 
diff --git a/Engine/Plugins/Runtime/Steam/SteamSockets/Source/SteamSockets/Private/SteamSocketsNetConnection.cpp b/Engine/Plugins/Runtime/Steam/SteamSockets/Source/SteamSockets/Private/SteamSocketsNetConnection.cpp
--- a/Engine/Plugins/Runtime/Steam/SteamSockets/Source/SteamSockets/Private/SteamSocketsNetConnection.cpp
+++ b/Engine/Plugins/Runtime/Steam/SteamSockets/Source/SteamSockets/Private/SteamSocketsNetConnection.cpp
@@ -90,7 +90,7 @@ void USteamSocketsNetConnection::LowLevelSend(void* Data, int32 CountBits, FOutP
 		SocketSub = static_cast<FSteamSocketsSubsystem*>(GetDriver()->GetSocketSubsystem());
 	}
 
-	if (ConnectionSocket != nullptr)
+	if (ConnectionSocket != nullptr && SocketSub)
 	{
 		const uint8* SendData = reinterpret_cast<const uint8*>(Data);
 
@@ -126,10 +126,7 @@ void USteamSocketsNetConnection::LowLevelSend(void* Data, int32 CountBits, FOutP
 		{
 			if (!ConnectionSocket->Send(SendData, BytesToSend, BytesSent))
 			{
-				if(SocketSub != nullptr)
-				{
-					UE_LOG(LogNet, Warning, TEXT("SteamSockets: LowLevelSend: Could not send %d bytes of data got error %d"), BytesToSend, (int32)SocketSub->GetLastErrorCode());
-				}
+				UE_LOG(LogNet, Warning, TEXT("SteamSockets: LowLevelSend: Could not send %d bytes of data got error %d"), BytesToSend, (int32)SocketSub->GetLastErrorCode());
 			}
 		}
 	}
diff --git a/Engine/Plugins/Runtime/Steam/SteamSockets/Source/SteamSockets/Private/SteamSocketsPrivate.h b/Engine/Plugins/Runtime/Steam/SteamSockets/Source/SteamSockets/Private/SteamSocketsPrivate.h
--- a/Engine/Plugins/Runtime/Steam/SteamSockets/Source/SteamSockets/Private/SteamSocketsPrivate.h
+++ b/Engine/Plugins/Runtime/Steam/SteamSockets/Source/SteamSockets/Private/SteamSocketsPrivate.h
@@ -22,3 +22,20 @@ THIRD_PARTY_INCLUDES_START
 #include "steam/steamnetworkingtypes.h"
 
 THIRD_PARTY_INCLUDES_END
+
+/** Some flags in the Steam SDK share similar numeral values which could change in the future.
+ *  Static analysis does not understand the difference and instead sees code like Object == 1 || Object == 1
+ *  For safety, we need to check both constants as the values in the SDK could change in a future version.
+ *  So this macro block was introduced to silence value redundancy warnings in static analysis.
+ */
+#if _MSC_VER
+#define STEAM_SDK_IGNORE_REDUNDANCY_START \
+		__pragma(warning(push)) \
+		__pragma(warning(disable: 6287)) 
+
+#define STEAM_SDK_IGNORE_REDUNDANCY_END \
+		__pragma(warning(pop))
+#else
+#define STEAM_SDK_IGNORE_REDUNDANCY_START
+#define STEAM_SDK_IGNORE_REDUNDANCY_END
+#endif
diff --git a/Engine/Plugins/Runtime/Steam/SteamSockets/Source/SteamSockets/Private/SteamSocketsSubsystem.cpp b/Engine/Plugins/Runtime/Steam/SteamSockets/Source/SteamSockets/Private/SteamSocketsSubsystem.cpp
--- a/Engine/Plugins/Runtime/Steam/SteamSockets/Source/SteamSockets/Private/SteamSocketsSubsystem.cpp
+++ b/Engine/Plugins/Runtime/Steam/SteamSockets/Source/SteamSockets/Private/SteamSocketsSubsystem.cpp
@@ -639,12 +639,14 @@ void FSteamSocketsSubsystem::AddSocket(const FInternetAddr& ForAddr, FSteamSocke
 		return;
 	}
 
+	STEAM_SDK_IGNORE_REDUNDANCY_START
 	// Do not attempt to add socket data for sockets that have invalid handles.
 	if (NewSocket->InternalHandle == k_HSteamListenSocket_Invalid || NewSocket->InternalHandle == k_HSteamNetConnection_Invalid)
 	{
 		UE_LOG(LogSockets, Verbose, TEXT("SteamSockets: Dropped socket tracking for socket with invalid handle."));
 		return;
 	}
+	STEAM_SDK_IGNORE_REDUNDANCY_END
 
 	FSteamSocketInformation NewSocketInfo(ForAddr.Clone(), NewSocket, ParentSocket);
 	if (SocketInformationMap.Find(NewSocket->InternalHandle) == nullptr)
@@ -706,11 +708,13 @@ void FSteamSocketsSubsystem::LinkNetDriver(FSocket* Socket, USteamSocketsNetDriv
 
 	FSteamSocket* SteamSocket = static_cast<FSteamSocket*>(Socket);
 	// Do not attempt to track netdrivers with invalid sockets.
+	STEAM_SDK_IGNORE_REDUNDANCY_START
 	if (SteamSocket->InternalHandle == k_HSteamListenSocket_Invalid || SteamSocket->InternalHandle == k_HSteamNetConnection_Invalid)
 	{
 		UE_LOG(LogSockets, Verbose, TEXT("SteamSockets: Dropped netdriver link with socket that has invalid handle."));
 		return;
 	}
+	STEAM_SDK_IGNORE_REDUNDANCY_END
 
 	// Link the netdriver to this socket.
 	FSteamSocketInformation* SocketInfo = GetSocketInfo(SteamSocket->InternalHandle);
