More CIS fixes

#rnx

diff --git a/Engine/Plugins/Runtime/Steam/SteamSockets/Source/SteamSockets/Private/SteamSocketsPrivate.h b/Engine/Plugins/Runtime/Steam/SteamSockets/Source/SteamSockets/Private/SteamSocketsPrivate.h
--- a/Engine/Plugins/Runtime/Steam/SteamSockets/Source/SteamSockets/Private/SteamSocketsPrivate.h
+++ b/Engine/Plugins/Runtime/Steam/SteamSockets/Source/SteamSockets/Private/SteamSocketsPrivate.h
@@ -28,7 +28,7 @@ THIRD_PARTY_INCLUDES_END
  *  For safety, we need to check both constants as the values in the SDK could change in a future version.
  *  So this macro block was introduced to silence value redundancy warnings in static analysis.
  */
-#if _MSC_VER
+#ifdef _MSC_VER
 #define STEAM_SDK_IGNORE_REDUNDANCY_START \
 		__pragma(warning(push)) \
 		__pragma(warning(disable: 6287)) 
