Fix -nosteam support such that it works properly with SteamSockets as one would expect.

#rnx

diff --git a/Engine/Plugins/Runtime/Steam/SteamSockets/Source/SteamSockets/Private/SteamSocketsModule.cpp b/Engine/Plugins/Runtime/Steam/SteamSockets/Source/SteamSockets/Private/SteamSocketsModule.cpp
--- a/Engine/Plugins/Runtime/Steam/SteamSockets/Source/SteamSockets/Private/SteamSocketsModule.cpp
+++ b/Engine/Plugins/Runtime/Steam/SteamSockets/Source/SteamSockets/Private/SteamSocketsModule.cpp
@@ -5,6 +5,8 @@
 #include "Misc/ConfigCacheIni.h"
 #include "Modules/ModuleManager.h"
 #include "OnlineSubsystemNames.h"
+#include "OnlineSubsystem.h"
+#include "OnlineSubsystemSteam.h"
 #include "SteamSharedModule.h"
 #include "SteamSocketsTypes.h"
 #include "SocketSubsystemModule.h"
@@ -19,11 +21,13 @@ namespace FNetworkProtocolTypes
 
 void FSteamSocketsModule::StartupModule()
 {
+	FOnlineSubsystemSteam* OnlineSteamSubsystem = static_cast<FOnlineSubsystemSteam*>(IOnlineSubsystem::Get(STEAM_SUBSYSTEM));
 	FSteamSharedModule& SharedModule = FSteamSharedModule::Get();
 	const bool bIsNotEditor = (IsRunningDedicatedServer() || IsRunningGame());
+	const bool bSteamOSSEnabled = (OnlineSteamSubsystem && OnlineSteamSubsystem->IsEnabled());
 
 	// Load the Steam modules before first call to API
-	if (SharedModule.AreSteamDllsLoaded() && bIsNotEditor)
+	if (SharedModule.AreSteamDllsLoaded() && bIsNotEditor && bSteamOSSEnabled)
 	{
 		// Settings flags
 		bool bOverrideSocketSubsystem = true;
@@ -51,6 +55,10 @@ void FSteamSocketsModule::StartupModule()
 			FSteamSocketsSubsystem::Destroy();
 		}
 	}
+	else if (!bSteamOSSEnabled)
+	{
+		UE_LOG(LogSockets, Log, TEXT("SteamSockets: Disabled due to no Steam OSS running."));
+	}
 	else if(bIsNotEditor)
 	{
 		UE_LOG(LogSockets, Warning, TEXT("SteamSockets: Steam SDK %s libraries not present at %s or failed to load!"), STEAM_SDK_VER, *SharedModule.GetSteamModulePath());
