Change the SteamSockets protocols to LazyNames as per all the other protocol types.

#rnx

diff --git a/Engine/Plugins/Runtime/Steam/SteamSockets/Source/SteamSockets/Private/SteamSocketsModule.cpp b/Engine/Plugins/Runtime/Steam/SteamSockets/Source/SteamSockets/Private/SteamSocketsModule.cpp
--- a/Engine/Plugins/Runtime/Steam/SteamSockets/Source/SteamSockets/Private/SteamSocketsModule.cpp
+++ b/Engine/Plugins/Runtime/Steam/SteamSockets/Source/SteamSockets/Private/SteamSocketsModule.cpp
@@ -15,8 +15,8 @@ IMPLEMENT_MODULE(FSteamSocketsModule, SteamSockets);
 
 namespace FNetworkProtocolTypes
 {
-	const FName SteamSocketsP2P(TEXT("SteamSocketsP2P"));
-	const FName SteamSocketsIP(TEXT("SteamSocketsIP"));
+	const FLazyName SteamSocketsP2P(TEXT("SteamSocketsP2P"));
+	const FLazyName SteamSocketsIP(TEXT("SteamSocketsIP"));
 }
 
 void FSteamSocketsModule::StartupModule()
diff --git a/Engine/Plugins/Runtime/Steam/SteamSockets/Source/SteamSockets/Public/SteamSocketsTypes.h b/Engine/Plugins/Runtime/Steam/SteamSockets/Source/SteamSockets/Public/SteamSocketsTypes.h
--- a/Engine/Plugins/Runtime/Steam/SteamSockets/Source/SteamSockets/Public/SteamSocketsTypes.h
+++ b/Engine/Plugins/Runtime/Steam/SteamSockets/Source/SteamSockets/Public/SteamSocketsTypes.h
@@ -11,8 +11,8 @@
 // Add our two new protocol types
 namespace FNetworkProtocolTypes
 {
-	STEAMSOCKETS_API extern const FName SteamSocketsP2P;
-	STEAMSOCKETS_API extern const FName SteamSocketsIP;
+	STEAMSOCKETS_API extern const FLazyName SteamSocketsP2P;
+	STEAMSOCKETS_API extern const FLazyName SteamSocketsIP;
 }
 
 typedef uint32 SteamSocketHandles;
