Adding fallbacks to SteamSocket identity address lookup as it should never fail on a logged in client. Also printing out the protocol address that's being searched for in the netdriver

#rnx

diff --git a/Engine/Plugins/Runtime/Steam/SteamSockets/Source/SteamSockets/Private/SteamSocketsNetDriver.cpp b/Engine/Plugins/Runtime/Steam/SteamSockets/Source/SteamSockets/Private/SteamSocketsNetDriver.cpp
--- a/Engine/Plugins/Runtime/Steam/SteamSockets/Source/SteamSockets/Private/SteamSocketsNetDriver.cpp
+++ b/Engine/Plugins/Runtime/Steam/SteamSockets/Source/SteamSockets/Private/SteamSocketsNetDriver.cpp
@@ -125,6 +125,8 @@ bool USteamSocketsNetDriver::InitBase(bool bInitAsClient, FNetworkNotify* InNoti
 	// Because of this, a normal flow should be skipped over and use a forced P2P address.
 	if (!bIsDelayedNetworkAccess)
 	{
+		UE_LOG(LogNet, Verbose, TEXT("SteamSockets: Looking for a binding address that matches protocol %s"), *SocketAddressType.ToString());
+
 		// Set the binding address.
 		TArray<TSharedPtr<FInternetAddr>> Addresses;
 		if (SocketSub->GetLocalAdapterAddresses(Addresses))
diff --git a/Engine/Plugins/Runtime/Steam/SteamSockets/Source/SteamSockets/Private/SteamSocketsSubsystem.cpp b/Engine/Plugins/Runtime/Steam/SteamSockets/Source/SteamSockets/Private/SteamSocketsSubsystem.cpp
--- a/Engine/Plugins/Runtime/Steam/SteamSockets/Source/SteamSockets/Private/SteamSocketsSubsystem.cpp
+++ b/Engine/Plugins/Runtime/Steam/SteamSockets/Source/SteamSockets/Private/SteamSocketsSubsystem.cpp
@@ -197,6 +197,7 @@ void FSteamSocketsSubsystem::Shutdown()
 	if (OnlineSteamSubsystem)
 	{
 		OnlineSteamSubsystem->ClearOnSteamServerLoginCompletedDelegate_Handle(SteamServerLoginDelegateHandle);
+		OnlineSteamSubsystem->SetPingInterface(nullptr);
 	}
 
 	// Clean up the internal event manager.
@@ -794,14 +795,27 @@ void FSteamSocketsSubsystem::OnServerLoginComplete(bool bWasSuccessful)
 TSharedPtr<FInternetAddr> FSteamSocketsSubsystem::GetIdentityAddress()
 {
 	SteamNetworkingIdentity SteamIDData;
+	TSharedRef<FInternetAddrSteamSockets> SteamAddr = StaticCastSharedRef<FInternetAddrSteamSockets>(CreateInternetAddr());
+
 	// Attempt to get the machine's identity (their steam ID)
 	if (GetSteamSocketsInterface() && GetSteamSocketsInterface()->GetIdentity(&SteamIDData))
 	{
 		// If we got it, assign it to the SteamAddr
-		TSharedRef<FInternetAddrSteamSockets> SteamAddr = StaticCastSharedRef<FInternetAddrSteamSockets>(CreateInternetAddr());
 		SteamAddr->Addr = SteamIDData;
 		return SteamAddr;
 	}
+	else if (!IsRunningDedicatedServer() && SteamUser())
+	{
+		CSteamID CurrentUser = SteamUser()->GetSteamID();
+		if (CurrentUser.IsValid())
+		{
+			SteamAddr->Addr.SetSteamID(CurrentUser);
+		}
+		else
+		{
+			UE_LOG(LogSockets, Warning, TEXT("SteamSockets: Unable to process current user's steam id!"));
+		}
+	}
 
 	return nullptr;
 }
