Updating array count macro after main merge in SteamSockets

#rnx

diff --git a/Engine/Plugins/Runtime/Steam/SteamSockets/Source/SteamSockets/Private/IPAddressSteamSockets.cpp b/Engine/Plugins/Runtime/Steam/SteamSockets/Source/SteamSockets/Private/IPAddressSteamSockets.cpp
--- a/Engine/Plugins/Runtime/Steam/SteamSockets/Source/SteamSockets/Private/IPAddressSteamSockets.cpp
+++ b/Engine/Plugins/Runtime/Steam/SteamSockets/Source/SteamSockets/Private/IPAddressSteamSockets.cpp
@@ -33,7 +33,7 @@ TArray<uint8> FInternetAddrSteamSockets::GetRawIp() const
 	{
 		const SteamNetworkingIPAddr* RawSteamIP = Addr.GetIPAddr();
 		RawAddressArray.Add(k_ESteamNetworkingIdentityType_IPAddress);
-		for (int32 i = 0; i < ARRAY_COUNT(RawSteamIP->m_ipv6); ++i)
+		for (int32 i = 0; i < UE_ARRAY_COUNT(RawSteamIP->m_ipv6); ++i)
 		{
 			RawAddressArray.Add(RawSteamIP->m_ipv6[i]);
 		}
