Adding a message to the steam sockets descriptor file that the SteamSocketsNetDriver is required.

#rnx

diff --git a/Engine/Plugins/Runtime/Steam/SteamSockets/SteamSockets.uplugin b/Engine/Plugins/Runtime/Steam/SteamSockets/SteamSockets.uplugin
--- a/Engine/Plugins/Runtime/Steam/SteamSockets/SteamSockets.uplugin
+++ b/Engine/Plugins/Runtime/Steam/SteamSockets/SteamSockets.uplugin
@@ -3,7 +3,7 @@
 	"Version" : 1.0,
 	"VersionName" : "1.0",
 	"FriendlyName" : "Steam Sockets",
-	"Description" : "New Steamworks Networking code that supports the SteamSockets interface.",
+	"Description" : "New Steamworks Networking code that supports the SteamSockets interface. NOTE: This plugin is only compatible with the SteamSockets Netdriver. It will not work if the proper netdriver definitions have not been set.",
 	"Category" : "Networking",
 	"CreatedBy" : "Epic Games, Inc.",
 	"CreatedByURL" : "https://epicgames.com",
