Fix the SteamSockets crash on Linux at exit. Appears to be caused by faults in how these background processes start on Linux. See JIRA for more comments.

#rnx

diff --git a/Engine/Plugins/Runtime/Steam/SteamSockets/Source/SteamSockets/Private/SteamSocketsSubsystem.cpp b/Engine/Plugins/Runtime/Steam/SteamSockets/Source/SteamSockets/Private/SteamSocketsSubsystem.cpp
--- a/Engine/Plugins/Runtime/Steam/SteamSockets/Source/SteamSockets/Private/SteamSocketsSubsystem.cpp
+++ b/Engine/Plugins/Runtime/Steam/SteamSockets/Source/SteamSockets/Private/SteamSocketsSubsystem.cpp
@@ -140,6 +140,9 @@ bool FSteamSocketsSubsystem::Init(FString& Error)
 		}
 	}
 
+	// These functions cause access violations on application exit due to how they spin up.
+	// When fixed in a new version of the Steam client application, this block can be removed.
+#if !PLATFORM_LINUX
 	// Clients and servers using the relay network will want to set this up
 	if (IsUsingRelayNetwork())
 	{
@@ -162,6 +165,7 @@ bool FSteamSocketsSubsystem::Init(FString& Error)
 	{
 		GetSteamSocketsInterface()->InitAuthentication();
 	}
+#endif
 	
 	// Set up the API Logging
 	ESteamNetworkingSocketsDebugOutputType DebugLevel = k_ESteamNetworkingSocketsDebugOutputType_Msg;
