New functionality to allow for creation of FInternetAddrs pre-set to a specific protocol at allocation. This allows for less ambiguity when obtaining an any/broadcast/loopback address.

diff --git a/Engine/Plugins/Online/OnlineSubsystem/Source/Private/LANBeacon.cpp b/Engine/Plugins/Online/OnlineSubsystem/Source/Private/LANBeacon.cpp
--- a/Engine/Plugins/Online/OnlineSubsystem/Source/Private/LANBeacon.cpp
+++ b/Engine/Plugins/Online/OnlineSubsystem/Source/Private/LANBeacon.cpp
@@ -43,7 +43,7 @@ bool FLanBeacon::Init(int32 Port)
 	ListenAddr = SocketSubsystem->GetLocalBindAddr(*GWarn);
 	ListenAddr->SetPort(Port);
 	// Set our broadcast address
-	BroadcastAddr = SocketSubsystem->CreateInternetAddr();
+	BroadcastAddr = SocketSubsystem->CreateInternetAddr(ListenAddr->GetProtocolType());
 	BroadcastAddr->SetBroadcastAddress();
 	BroadcastAddr->SetPort(Port);
 	// A temporary "received from" address
diff --git a/Engine/Plugins/Runtime/Steam/SteamSockets/Source/SteamSockets/Private/IPAddressSteamSockets.cpp b/Engine/Plugins/Runtime/Steam/SteamSockets/Source/SteamSockets/Private/IPAddressSteamSockets.cpp
--- a/Engine/Plugins/Runtime/Steam/SteamSockets/Source/SteamSockets/Private/IPAddressSteamSockets.cpp
+++ b/Engine/Plugins/Runtime/Steam/SteamSockets/Source/SteamSockets/Private/IPAddressSteamSockets.cpp
@@ -68,6 +68,7 @@ void FInternetAddrSteamSockets::SetRawIp(const TArray<uint8>& RawAddr)
 		}
 
 		Addr.SetSteamID64(NewSteamId);
+		ProtocolType = FNetworkProtocolTypes::SteamSocketsP2P;
 	}
 	else if(ArrayType == k_ESteamNetworkingIdentityType_IPAddress)
 	{
@@ -81,6 +82,7 @@ void FInternetAddrSteamSockets::SetRawIp(const TArray<uint8>& RawAddr)
 		}
 
 		Addr.SetIPAddr(NewAddr);
+		ProtocolType = FNetworkProtocolTypes::SteamSocketsIP;
 	}
 }
 
@@ -105,6 +107,7 @@ void FInternetAddrSteamSockets::SetIp(const TCHAR* InAddr, bool& bIsValid)
 					P2PVirtualPort = Channel;
 					bIsValid = true;
 				}
+				ProtocolType = FNetworkProtocolTypes::SteamSocketsP2P;
 			}
 			else
 			{
@@ -115,6 +118,7 @@ void FInternetAddrSteamSockets::SetIp(const TCHAR* InAddr, bool& bIsValid)
 		else
 		{
 			Addr.SetSteamID64(FCString::Atoi64(*InAddrStr));
+			ProtocolType = FNetworkProtocolTypes::SteamSocketsP2P;
 			bIsValid = true;
 			return;
 		}
@@ -124,6 +128,7 @@ void FInternetAddrSteamSockets::SetIp(const TCHAR* InAddr, bool& bIsValid)
 		// This is an IP address. Tell Steam to parse it.
 		SteamNetworkingIPAddr NewAddress;
 		bIsValid = NewAddress.ParseString(TCHAR_TO_ANSI(InAddr));
+		ProtocolType = FNetworkProtocolTypes::SteamSocketsIP;
 		Addr.SetIPAddr(NewAddress);
 	}
 }
@@ -169,6 +174,7 @@ void FInternetAddrSteamSockets::SetAnyAddress()
 	Addr.Clear();
 	SteamNetworkingIPAddr NewAddress;
 	NewAddress.Clear();
+	ProtocolType = FNetworkProtocolTypes::SteamSocketsIP;
 	Addr.SetIPAddr(NewAddress);
 }
 
diff --git a/Engine/Plugins/Runtime/Steam/SteamSockets/Source/SteamSockets/Private/IPAddressSteamSockets.h b/Engine/Plugins/Runtime/Steam/SteamSockets/Source/SteamSockets/Private/IPAddressSteamSockets.h
--- a/Engine/Plugins/Runtime/Steam/SteamSockets/Source/SteamSockets/Private/IPAddressSteamSockets.h
+++ b/Engine/Plugins/Runtime/Steam/SteamSockets/Source/SteamSockets/Private/IPAddressSteamSockets.h
@@ -14,34 +14,40 @@ class FInternetAddrSteamSockets : public FInternetAddr
 PACKAGE_SCOPE:
 	SteamNetworkingIdentity Addr;
 	int32 P2PVirtualPort;
+	FName ProtocolType;
 
 public:
-	FInternetAddrSteamSockets() :
-		P2PVirtualPort(0)
+	FInternetAddrSteamSockets(const FName RequestedProtocol = NAME_None) :
+		P2PVirtualPort(0),
+		ProtocolType(RequestedProtocol)
 	{
 		Addr.Clear();
 	}
 
 	FInternetAddrSteamSockets(const FInternetAddrSteamSockets& In) :
 		Addr(In.Addr),
-		P2PVirtualPort(In.P2PVirtualPort)
+		P2PVirtualPort(In.P2PVirtualPort),
+		ProtocolType(In.ProtocolType)
 	{
 	}
 
 	FInternetAddrSteamSockets(const SteamNetworkingIdentity& NewAddress) :
 		Addr(NewAddress),
-		P2PVirtualPort(0)
+		P2PVirtualPort(0),
+		ProtocolType(NewAddress.GetIPAddr() == nullptr ? FNetworkProtocolTypes::SteamSocketsP2P : FNetworkProtocolTypes::SteamSocketsIP)
 	{
 	}
 
 	FInternetAddrSteamSockets(const SteamNetworkingIPAddr& IPAddr) :
-		P2PVirtualPort(0)
+		P2PVirtualPort(0),
+		ProtocolType(FNetworkProtocolTypes::SteamSocketsIP)
 	{
 		Addr.SetIPAddr(IPAddr);
 	}
 
 	explicit FInternetAddrSteamSockets(uint64& SteamID) :
-		P2PVirtualPort(0)
+		P2PVirtualPort(0),
+		ProtocolType(FNetworkProtocolTypes::SteamSocketsP2P)
 	{
 		Addr.SetSteamID64(SteamID);
 	}
@@ -161,7 +167,7 @@ public:
 
 	virtual TSharedRef<FInternetAddr> Clone() const override
 	{
-		TSharedRef<FInternetAddrSteamSockets> NewAddress = MakeShareable(new FInternetAddrSteamSockets);
+		TSharedRef<FInternetAddrSteamSockets> NewAddress = MakeShareable(new FInternetAddrSteamSockets(ProtocolType));
 		NewAddress->Addr = Addr;
 		NewAddress->P2PVirtualPort = P2PVirtualPort;
 		return NewAddress;
diff --git a/Engine/Plugins/Runtime/Steam/SteamSockets/Source/SteamSockets/Private/SteamSocketsSubsystem.cpp b/Engine/Plugins/Runtime/Steam/SteamSockets/Source/SteamSockets/Private/SteamSocketsSubsystem.cpp
--- a/Engine/Plugins/Runtime/Steam/SteamSockets/Source/SteamSockets/Private/SteamSocketsSubsystem.cpp
+++ b/Engine/Plugins/Runtime/Steam/SteamSockets/Source/SteamSockets/Private/SteamSocketsSubsystem.cpp
@@ -400,6 +400,11 @@ TSharedRef<FInternetAddr> FSteamSocketsSubsystem::CreateInternetAddr()
 	return MakeShareable(new FInternetAddrSteamSockets());
 }
 
+TSharedRef<FInternetAddr> FSteamSocketsSubsystem::CreateInternetAddr(const FName RequestedProtocol)
+{
+	return MakeShareable(new FInternetAddrSteamSockets(RequestedProtocol));
+}
+
 const TCHAR* FSteamSocketsSubsystem::GetSocketAPIName() const
 {
 	return TEXT("SteamSockets");
@@ -408,7 +413,7 @@ const TCHAR* FSteamSocketsSubsystem::GetSocketAPIName() const
 bool FSteamSocketsSubsystem::GetLocalAdapterAddresses(TArray<TSharedPtr<FInternetAddr>>& OutAddresses)
 {
 	// Add Multihome
-	TSharedRef<FInternetAddr> MultihomeAddress = CreateInternetAddr();
+	TSharedRef<FInternetAddr> MultihomeAddress = CreateInternetAddr(FNetworkProtocolTypes::SteamSocketsIP);
 	if (GetMultihomeAddress(MultihomeAddress))
 	{
 		OutAddresses.Add(MultihomeAddress);
@@ -425,7 +430,7 @@ bool FSteamSocketsSubsystem::GetLocalAdapterAddresses(TArray<TSharedPtr<FInterne
 	}
 
 	// Always include the any address.
-	TSharedRef<FInternetAddr> AnyAddress = CreateInternetAddr();
+	TSharedRef<FInternetAddr> AnyAddress = CreateInternetAddr(FNetworkProtocolTypes::SteamSocketsIP);
 	AnyAddress->SetAnyAddress();
 	OutAddresses.Add(AnyAddress);
 
@@ -440,7 +445,7 @@ TSharedRef<FInternetAddr> FSteamSocketsSubsystem::GetLocalBindAddr(FOutputDevice
 		return Adapters[0].ToSharedRef();
 	}
 
-	return CreateInternetAddr();
+	return CreateInternetAddr(FNetworkProtocolTypes::SteamSocketsIP);
 }
 
 void FSteamSocketsSubsystem::CleanSocketInformation(bool bForceClean)
@@ -773,9 +778,13 @@ void FSteamSocketsSubsystem::OnServerLoginComplete(bool bWasSuccessful)
 		if (SteamIdentityAddr.IsValid())
 		{
 			// Push in the new valid information.
-			DelayedListener.Socket->BindAddress.Addr = SteamIdentityAddr->Addr;
-			SteamNetDriver->LocalAddr = SteamIdentityAddr->Clone();
-			SteamNetDriver->LocalAddr->SetPlatformPort(DelayedListener.Socket->BindAddress.GetPlatformPort());
+			int32 ListenPort = DelayedListener.Socket->BindAddress.GetPlatformPort();
+			TSharedPtr<FInternetAddrSteamSockets> ListenerAddress = StaticCastSharedRef<FInternetAddrSteamSockets>(SteamIdentityAddr->Clone());
+			ListenerAddress->SetPlatformPort(ListenPort);
+
+			// Update the socket and netdriver addresses
+			DelayedListener.Socket->BindAddress = *ListenerAddress;
+			SteamNetDriver->LocalAddr = ListenerAddress;
 
 			// Actually start the listen call
 			if (!DelayedListener.Socket->Listen(0))
@@ -801,7 +810,7 @@ void FSteamSocketsSubsystem::OnServerLoginComplete(bool bWasSuccessful)
 TSharedPtr<FInternetAddr> FSteamSocketsSubsystem::GetIdentityAddress()
 {
 	SteamNetworkingIdentity SteamIDData;
-	TSharedRef<FInternetAddrSteamSockets> SteamAddr = StaticCastSharedRef<FInternetAddrSteamSockets>(CreateInternetAddr());
+	TSharedRef<FInternetAddrSteamSockets> SteamAddr = StaticCastSharedRef<FInternetAddrSteamSockets>(CreateInternetAddr(FNetworkProtocolTypes::SteamSocketsP2P));
 
 	// Attempt to get the machine's identity (their steam ID)
 	if (GetSteamSocketsInterface() && GetSteamSocketsInterface()->GetIdentity(&SteamIDData))
diff --git a/Engine/Plugins/Runtime/Steam/SteamSockets/Source/SteamSockets/Public/SteamSocketsSubsystem.h b/Engine/Plugins/Runtime/Steam/SteamSockets/Source/SteamSockets/Public/SteamSocketsSubsystem.h
--- a/Engine/Plugins/Runtime/Steam/SteamSockets/Source/SteamSockets/Public/SteamSocketsSubsystem.h
+++ b/Engine/Plugins/Runtime/Steam/SteamSockets/Source/SteamSockets/Public/SteamSocketsSubsystem.h
@@ -53,6 +53,7 @@ public:
 	virtual bool GetHostName(FString& HostName) override;
 
 	virtual TSharedRef<FInternetAddr> CreateInternetAddr() override;
+	virtual TSharedRef<FInternetAddr> CreateInternetAddr(const FName RequestedProtocol) override;
 
 	virtual const TCHAR* GetSocketAPIName() const override;
 
diff --git a/Engine/Source/Runtime/Networking/Public/Common/UdpSocketBuilder.h b/Engine/Source/Runtime/Networking/Public/Common/UdpSocketBuilder.h
--- a/Engine/Source/Runtime/Networking/Public/Common/UdpSocketBuilder.h
+++ b/Engine/Source/Runtime/Networking/Public/Common/UdpSocketBuilder.h
@@ -318,7 +318,7 @@ public:
 		}
 		else
 		{
-			TSharedRef<FInternetAddr> MulticastAddress = SocketSubsystem->CreateInternetAddr();
+			TSharedRef<FInternetAddr> MulticastAddress = SocketSubsystem->CreateInternetAddr(RemoteAddr->GetProtocolType());
 			MulticastAddress->SetBroadcastAddress();
 			TSharedPtr<FInternetAddr> AddressToUse = nullptr;
 
diff --git a/Engine/Source/Runtime/Online/HTTP/Private/Curl/CurlHttpManager.cpp b/Engine/Source/Runtime/Online/HTTP/Private/Curl/CurlHttpManager.cpp
--- a/Engine/Source/Runtime/Online/HTTP/Private/Curl/CurlHttpManager.cpp
+++ b/Engine/Source/Runtime/Online/HTTP/Private/Curl/CurlHttpManager.cpp
@@ -263,17 +263,9 @@ void FCurlHttpManager::InitCurl()
 	if (FParse::Value(FCommandLine::Get(), TEXT("MULTIHOMEHTTP="), Home, UE_ARRAY_COUNT(Home)))
 	{
 		ISocketSubsystem* SocketSubsystem = ISocketSubsystem::Get(PLATFORM_SOCKETSUBSYSTEM);
-		if (SocketSubsystem)
+		if (SocketSubsystem && SocketSubsystem->GetAddressFromString(Home).IsValid())
 		{
-			TSharedRef<FInternetAddr> HostAddr = SocketSubsystem->CreateInternetAddr();
-			HostAddr->SetAnyAddress();
-
-			bool bIsValid = false;
-			HostAddr->SetIp(Home, bIsValid);
-			if (bIsValid)
-			{
-				CurlRequestOptions.LocalHostAddr = FString(Home);
-			}
+			CurlRequestOptions.LocalHostAddr = FString(Home);
 		}
 	}
 
diff --git a/Engine/Source/Runtime/Sockets/Private/BSDSockets/IPAddressBSD.cpp b/Engine/Source/Runtime/Sockets/Private/BSDSockets/IPAddressBSD.cpp
--- a/Engine/Source/Runtime/Sockets/Private/BSDSockets/IPAddressBSD.cpp
+++ b/Engine/Source/Runtime/Sockets/Private/BSDSockets/IPAddressBSD.cpp
@@ -24,9 +24,30 @@ FInternetAddrBSD::FInternetAddrBSD()
 	Clear();
 }
 
-FInternetAddrBSD::FInternetAddrBSD(FSocketSubsystemBSD* InSocketSubsystem) : SocketSubsystem(InSocketSubsystem)
+FInternetAddrBSD::FInternetAddrBSD(FSocketSubsystemBSD* InSocketSubsystem, FName RequestedProtocol) : 
+	SocketSubsystem(InSocketSubsystem)
 {
 	Clear();
+
+	if (RequestedProtocol.IsNone() && InSocketSubsystem)
+	{
+		RequestedProtocol = InSocketSubsystem->GetDefaultSocketProtocolFamily();
+	}
+
+	if (RequestedProtocol == FNetworkProtocolTypes::IPv4)
+	{
+		Addr.ss_family = AF_INET;
+	}
+#if PLATFORM_HAS_BSD_IPV6_SOCKETS
+	else if (RequestedProtocol == FNetworkProtocolTypes::IPv6)
+	{
+		Addr.ss_family = AF_INET6;
+	}
+#endif
+	else
+	{
+		Addr.ss_family = AF_UNSPEC;
+	}
 }
 
 bool FInternetAddrBSD::CompareEndpoints(const FInternetAddr& InAddr) const
@@ -180,7 +201,7 @@ void FInternetAddrBSD::SetIp(const TCHAR* InAddr, bool& bIsValid)
 void FInternetAddrBSD::SetIp(uint32 InAddr)
 {
 #if PLATFORM_HAS_BSD_IPV6_SOCKETS
-	if (SocketSubsystem && SocketSubsystem->GetDefaultSocketProtocolFamily() == FNetworkProtocolTypes::IPv6)
+	if (GetProtocolType() == FNetworkProtocolTypes::IPv6)
 	{
 		if (InAddr == 0)
 		{
@@ -319,7 +340,7 @@ void FInternetAddrBSD::SetAnyAddress()
 {
 	if (SocketSubsystem != nullptr)
 	{
-		if (SocketSubsystem->GetDefaultSocketProtocolFamily() == FNetworkProtocolTypes::IPv6)
+		if (GetProtocolType() == FNetworkProtocolTypes::IPv6)
 		{
 			SetAnyIPv6Address();
 		}
@@ -353,7 +374,7 @@ void FInternetAddrBSD::SetBroadcastAddress()
 {
 	if (SocketSubsystem)
 	{
-		if (SocketSubsystem->GetDefaultSocketProtocolFamily() == FNetworkProtocolTypes::IPv6)
+		if (GetProtocolType() == FNetworkProtocolTypes::IPv6)
 		{
 			SetIPv6BroadcastAddress();
 		}
@@ -395,7 +416,7 @@ void FInternetAddrBSD::SetLoopbackAddress()
 {
 	if (SocketSubsystem)
 	{
-		if (SocketSubsystem->GetDefaultSocketProtocolFamily() == FNetworkProtocolTypes::IPv6)
+		if (GetProtocolType() == FNetworkProtocolTypes::IPv6)
 		{
 			SetIPv6LoopbackAddress();
 		}
@@ -523,7 +544,7 @@ bool FInternetAddrBSD::IsValid() const
 
 TSharedRef<FInternetAddr> FInternetAddrBSD::Clone() const
 {
-	TSharedRef<FInternetAddrBSD> NewAddress = MakeShareable(new FInternetAddrBSD(SocketSubsystem));
+	TSharedRef<FInternetAddrBSD> NewAddress = MakeShareable(new FInternetAddrBSD(SocketSubsystem, GetProtocolType()));
 	NewAddress->Addr = Addr;
 	return NewAddress;
 }
diff --git a/Engine/Source/Runtime/Sockets/Private/BSDSockets/IPAddressBSD.h b/Engine/Source/Runtime/Sockets/Private/BSDSockets/IPAddressBSD.h
--- a/Engine/Source/Runtime/Sockets/Private/BSDSockets/IPAddressBSD.h
+++ b/Engine/Source/Runtime/Sockets/Private/BSDSockets/IPAddressBSD.h
@@ -60,7 +60,7 @@ public:
 	 */
 	FInternetAddrBSD();
 
-	FInternetAddrBSD(FSocketSubsystemBSD* InSocketSubsystem);
+	FInternetAddrBSD(FSocketSubsystemBSD* InSocketSubsystem, FName RequestedProtocol=NAME_None);
 
 	/**
 	 * Compares FInternetAddrs together, comparing the logical net addresses (endpoints)
@@ -187,8 +187,8 @@ public:
 
 	/**
 	 * Sets the address structure to be bindable to any IP address.
-	 * Uses the value returned from GetDefaultSocketProtocolFamily to determine default
-	 * addresses to be used.
+	 * Uses the protocol the address structure was created with (or currently set to) to
+	 * determine the address protocol type used.
 	 *
 	 * To skip assumptions, you can call the designated version explicitly below.
 	 */
@@ -202,8 +202,8 @@ public:
 
 	/**
 	 * Sets the address structure to be bound to the multicast ip address.
-	 * Uses the value returned from GetDefaultSocketProtocolFamily to determine default
-	 * addresses to be used.
+	 * Uses the protocol the address structure was created with (or currently set to) to
+	 * determine the address protocol type used.
 	 *
 	 * To skip assumptions, you can call the designated version explicitly below.
 	 */
@@ -217,8 +217,8 @@ public:
 
 	/**
 	 * Sets the address structure to be bound to the loopback ip address.
-	 * Uses the value returned from GetDefaultSocketProtocolFamily to determine default
-	 * addresses to be used.
+	 * Uses the protocol the address structure was created with (or currently set to) to
+	 * determine the address protocol type used.
 	 *
 	 * To skip assumptions, you can call the designated version explicitly below.
 	 */
diff --git a/Engine/Source/Runtime/Sockets/Private/BSDSockets/SocketSubsystemBSD.cpp b/Engine/Source/Runtime/Sockets/Private/BSDSockets/SocketSubsystemBSD.cpp
--- a/Engine/Source/Runtime/Sockets/Private/BSDSockets/SocketSubsystemBSD.cpp
+++ b/Engine/Source/Runtime/Sockets/Private/BSDSockets/SocketSubsystemBSD.cpp
@@ -244,6 +244,11 @@ TSharedRef<FInternetAddr> FSocketSubsystemBSD::CreateInternetAddr()
 	return MakeShareable(new FInternetAddrBSD(this));
 }
 
+TSharedRef<FInternetAddr> FSocketSubsystemBSD::CreateInternetAddr(const FName ProtocolType)
+{
+	return MakeShareable(new FInternetAddrBSD(this, ProtocolType));
+}
+
 bool FSocketSubsystemBSD::IsSocketWaitSupported() const
 {
 	return true;
diff --git a/Engine/Source/Runtime/Sockets/Private/BSDSockets/SocketSubsystemBSD.h b/Engine/Source/Runtime/Sockets/Private/BSDSockets/SocketSubsystemBSD.h
--- a/Engine/Source/Runtime/Sockets/Private/BSDSockets/SocketSubsystemBSD.h
+++ b/Engine/Source/Runtime/Sockets/Private/BSDSockets/SocketSubsystemBSD.h
@@ -38,6 +38,7 @@ public:
 
 	// ISocketSubsystem interface
 	virtual TSharedRef<FInternetAddr> CreateInternetAddr() override;
+	virtual TSharedRef<FInternetAddr> CreateInternetAddr(const FName ProtocolType) override;
 
 	virtual class FSocket* CreateSocket(const FName& SocketType, const FString& SocketDescription, bool bForceUDP = false) override
 	{
diff --git a/Engine/Source/Runtime/Sockets/Private/IOS/IPAddressBSDIOS.h b/Engine/Source/Runtime/Sockets/Private/IOS/IPAddressBSDIOS.h
--- a/Engine/Source/Runtime/Sockets/Private/IOS/IPAddressBSDIOS.h
+++ b/Engine/Source/Runtime/Sockets/Private/IOS/IPAddressBSDIOS.h
@@ -11,10 +11,12 @@
 
 class FInternetAddrBSDIOS : public FInternetAddrBSD
 {
-
 public:
 
-	FInternetAddrBSDIOS(FSocketSubsystemBSD* InSocketSubsystem) : FInternetAddrBSD(InSocketSubsystem) {}
+	FInternetAddrBSDIOS(FSocketSubsystemBSD* InSocketSubsystem, FName RequestedProtocol = NAME_None) : 
+		FInternetAddrBSD(InSocketSubsystem, RequestedProtocol)
+	{
+	}
 
 	/** Sets the address to broadcast */
 	virtual void SetIPv6BroadcastAddress() override
diff --git a/Engine/Source/Runtime/Sockets/Private/IOS/SocketSubsystemIOS.cpp b/Engine/Source/Runtime/Sockets/Private/IOS/SocketSubsystemIOS.cpp
--- a/Engine/Source/Runtime/Sockets/Private/IOS/SocketSubsystemIOS.cpp
+++ b/Engine/Source/Runtime/Sockets/Private/IOS/SocketSubsystemIOS.cpp
@@ -178,3 +178,8 @@ TSharedRef<FInternetAddr> FSocketSubsystemIOS::CreateInternetAddr()
 {
 	return MakeShareable(new FInternetAddrBSDIOS(this));
 }
+
+TSharedRef<FInternetAddr> FSocketSubsystemIOS::CreateInternetAddr(const FName RequiredProtocol)
+{
+	return MakeShareable(new FInternetAddrBSDIOS(this, RequiredProtocol));
+}
diff --git a/Engine/Source/Runtime/Sockets/Private/IOS/SocketSubsystemIOS.h b/Engine/Source/Runtime/Sockets/Private/IOS/SocketSubsystemIOS.h
--- a/Engine/Source/Runtime/Sockets/Private/IOS/SocketSubsystemIOS.h
+++ b/Engine/Source/Runtime/Sockets/Private/IOS/SocketSubsystemIOS.h
@@ -55,6 +55,7 @@ public:
 	virtual bool HasNetworkDevice() override;
 	virtual FSocket* CreateSocket(const FName& SocketType, const FString& SocketDescription, const FName& ProtocolType) override;
 	virtual TSharedRef<FInternetAddr> CreateInternetAddr() override;
+	virtual TSharedRef<FInternetAddr> CreateInternetAddr(const FName RequiredProtocol) override;
 	virtual TSharedRef<FInternetAddr> GetLocalHostAddr(FOutputDevice& Out, bool& bCanBindAll) override;
 	virtual class FSocketBSD* InternalBSDSocketFactory(SOCKET Socket, ESocketType SocketType, const FString& SocketDescription, const FName& SocketProtocol) override;
 };
\ No newline at end of file
diff --git a/Engine/Source/Runtime/Sockets/Public/SocketSubsystem.h b/Engine/Source/Runtime/Sockets/Public/SocketSubsystem.h
--- a/Engine/Source/Runtime/Sockets/Public/SocketSubsystem.h
+++ b/Engine/Source/Runtime/Sockets/Public/SocketSubsystem.h
@@ -291,6 +291,15 @@ public:
 	 */
 	virtual TSharedRef<FInternetAddr> CreateInternetAddr() = 0;
 
+	/**
+	 * Create a FInternetAddr of the desired protocol
+	 */
+	virtual TSharedRef<FInternetAddr> CreateInternetAddr(const FName ProtocolType)
+	{
+		// If not implemented, returns the base version
+		return CreateInternetAddr();
+	}
+
 	/**
 	 * Create a platform specific FRecvMulti representation
 	 *
