Fix ICE in SteamSocketsSubsystem.cpp when using newer VC toolchains.

#rnx

diff --git a/Engine/Plugins/Runtime/Steam/SteamSockets/Source/SteamSockets/Private/SteamSocketsSubsystem.cpp b/Engine/Plugins/Runtime/Steam/SteamSockets/Source/SteamSockets/Private/SteamSocketsSubsystem.cpp
--- a/Engine/Plugins/Runtime/Steam/SteamSockets/Source/SteamSockets/Private/SteamSocketsSubsystem.cpp
+++ b/Engine/Plugins/Runtime/Steam/SteamSockets/Source/SteamSockets/Private/SteamSocketsSubsystem.cpp
@@ -187,9 +187,7 @@ bool FSteamSocketsSubsystem::Init(FString& Error)
 	if (OnlineSteamSubsystem != nullptr && IsUsingRelayNetwork())
 	{
 		// Deletion will be handled by the Steam OSS, we do not need to worry about it in this subsystem.
-		TSharedPtr<FSteamSocketsPing, ESPMode::ThreadSafe> SteamPingHandler = MakeShareable(new FSteamSocketsPing(this, OnlineSteamSubsystem));
-		FOnlinePingSteamPtr NewPingInterface = StaticCastSharedPtr<FOnlinePingInterfaceSteam>(SteamPingHandler);
-		OnlineSteamSubsystem->SetPingInterface(NewPingInterface);
+		OnlineSteamSubsystem->SetPingInterface(MakeShared<FSteamSocketsPing, ESPMode::ThreadSafe>(this, OnlineSteamSubsystem));
 	}
 
 	return true;
