Fixed mismatched copy constructor/assignment operator in FInternetAddrSteamSockets.

diff --git a/Engine/Plugins/Runtime/Steam/SteamSockets/Source/SteamSockets/Private/IPAddressSteamSockets.h b/Engine/Plugins/Runtime/Steam/SteamSockets/Source/SteamSockets/Private/IPAddressSteamSockets.h
--- a/Engine/Plugins/Runtime/Steam/SteamSockets/Source/SteamSockets/Private/IPAddressSteamSockets.h
+++ b/Engine/Plugins/Runtime/Steam/SteamSockets/Source/SteamSockets/Private/IPAddressSteamSockets.h
@@ -28,13 +28,6 @@ public:
 		Addr.Clear();
 	}
 
-	FInternetAddrSteamSockets(const FInternetAddrSteamSockets& In) :
-		Addr(In.Addr),
-		P2PVirtualPort(In.P2PVirtualPort),
-		ProtocolType(In.ProtocolType)
-	{
-	}
-
 	FInternetAddrSteamSockets(const SteamNetworkingIdentity& NewAddress) :
 		Addr(NewAddress),
 		P2PVirtualPort(0),
