SteamSockets fixes

SteamSocketsNetConnection has a custom handshake flow where it performs the handshake using the handler of the net driver, but this was causing problems after it was complete because the handler of the net connection did not have the data of the handshake (namely, the session and client IDs). This change synchronizes the two after the handshake is complete

There were also some non-unity issues in the Steam sockets files that have been addressed.

diff --git a/Engine/Plugins/Runtime/Steam/SteamSockets/Source/SteamSockets/Private/SteamSocketsNetConnection.cpp b/Engine/Plugins/Runtime/Steam/SteamSockets/Source/SteamSockets/Private/SteamSocketsNetConnection.cpp
--- a/Engine/Plugins/Runtime/Steam/SteamSockets/Source/SteamSockets/Private/SteamSocketsNetConnection.cpp
+++ b/Engine/Plugins/Runtime/Steam/SteamSockets/Source/SteamSockets/Private/SteamSocketsNetConnection.cpp
@@ -1,6 +1,7 @@
 // Copyright Epic Games, Inc. All Rights Reserved.
 
 #include "SteamSocketsNetConnection.h"
+#include "SteamSocketsPackage.h"
 #include "SteamSocketsSubsystem.h"
 #include "SteamSocketsNetDriver.h"
 #include "SteamSocket.h"
@@ -200,6 +201,12 @@ void USteamSocketsNetConnection::HandleRecvMessage(void* InData, int32 SizeOfDat
 
 				UE_LOG(LogNet, Log, TEXT("SteamSockets: Connectionless handshake complete"));
 				bInConnectionlessHandshake = false;
+
+				// Transfer net driver's stateless connection data to the net connection
+				if(StatelessConnectComponent.IsValid())
+				{
+					StatelessConnectComponent.Pin()->SetDriver(SteamNetDriver);
+				}
 				
 				// Reset the challenge data for the future
 				if (StatelessConnect.IsValid())
diff --git a/Engine/Plugins/Runtime/Steam/SteamSockets/Source/SteamSockets/Private/SteamSocketsSubsystem.cpp b/Engine/Plugins/Runtime/Steam/SteamSockets/Source/SteamSockets/Private/SteamSocketsSubsystem.cpp
--- a/Engine/Plugins/Runtime/Steam/SteamSockets/Source/SteamSockets/Private/SteamSocketsSubsystem.cpp
+++ b/Engine/Plugins/Runtime/Steam/SteamSockets/Source/SteamSockets/Private/SteamSocketsSubsystem.cpp
@@ -1,6 +1,7 @@
 // Copyright Epic Games, Inc. All Rights Reserved.
 
 #include "SteamSocketsSubsystem.h"
+#include "SteamSocketsPackage.h"
 #include "Misc/CommandLine.h"
 #include "Misc/ConfigCacheIni.h"
 #include "SteamSocketsTaskManager.h"
